DaySequence: Move DaySequence from Restricted/NFL to Experimental

diff --git a/Engine/Plugins/Experimental/DaySequence/Source/DaySequence/Private/Actors/DaySequenceModifierVolume.cpp b/Engine/Plugins/Experimental/DaySequence/Source/DaySequence/Private/Actors/DaySequenceModifierVolume.cpp
new file mode 100644
--- /dev/null
+++ b/Engine/Plugins/Experimental/DaySequence/Source/DaySequence/Private/Actors/DaySequenceModifierVolume.cpp
@@ -0,0 +1,294 @@
+// Copyright Epic Games, Inc. All Rights Reserved.
+
+#include "Actors/DaySequenceModifierVolume.h"
+
+#include "DaySequenceActor.h"
+#include "DaySequenceModifierComponent.h"
+#include "DaySequenceSubsystem.h"
+
+#include "Components/BoxComponent.h"
+#include "Engine/World.h"
+#include "GameDelegates.h"
+#include "GameFramework/Pawn.h"
+#include "TimerManager.h"
+
+ADaySequenceModifierVolume::ADaySequenceModifierVolume(const FObjectInitializer& Init)
+: Super(Init)
+{
+	PrimaryActorTick.bCanEverTick = true;
+	
+	DaySequenceModifier = CreateDefaultSubobject<UDaySequenceModifierComponent>(TEXT("DaySequenceModifier"));
+	DaySequenceModifier->SetupAttachment(RootComponent);
+
+	DefaultBox = CreateDefaultSubobject<UBoxComponent>(TEXT("Box"));
+	DefaultBox->SetupAttachment(DaySequenceModifier);
+	DefaultBox->SetLineThickness(10.f);
+	DefaultBox->SetBoxExtent(FVector(500.f));
+}
+
+void ADaySequenceModifierVolume::TryEnableModifier() const
+{	
+	if (DaySequenceModifier->IsBlendTargetInAnyVolume())
+	{
+		DaySequenceModifier->EnableModifier();
+	}
+	else
+	{
+		DaySequenceModifier->DisableModifier();
+	}
+}
+
+void ADaySequenceModifierVolume::SetBlendTarget(AActor* InBlendTarget)
+{
+	if (!IsValid(InBlendTarget) || InBlendTarget == CurrentBlendTarget)
+	{
+		return;
+	}
+
+	CurrentBlendTarget = InBlendTarget;
+
+	DaySequenceModifier->EnableDistanceVolumeBlends(InBlendTarget);
+	DaySequenceModifier->SetCustomVolumeBlendWeight(1.f);
+
+	TryEnableModifier();
+}
+
+void ADaySequenceModifierVolume::BeginPlay()
+{
+	Super::BeginPlay();
+
+	Initialize();
+}
+
+void ADaySequenceModifierVolume::OnConstruction(const FTransform& Transform)
+{
+	Super::OnConstruction(Transform);
+	
+	Initialize();
+}
+
+void ADaySequenceModifierVolume::Initialize()
+{
+	if (IsTemplate())
+	{
+		return;
+	}
+
+	// This actor should only initialize on the client.
+	if (GetNetMode() == NM_DedicatedServer)
+	{
+		SetActorEnableCollision(false);
+		return;
+	}
+
+#if WITH_EDITOR
+	const UWorld* World = GetWorld();
+	if (World && World->WorldType == EWorldType::Editor)
+	{
+		DaySequenceActor = nullptr;
+		if (IsValid(DaySequenceModifier))
+		{
+			DaySequenceModifier->UnbindFromDaySequenceActor();
+		}
+	}
+#endif
+
+	DaySequenceActorSetup();
+}
+
+void ADaySequenceModifierVolume::VolumeSetup()
+{
+	DaySequenceModifier->EmptyVolumeShapeComponents();
+	
+	AddShapeComponentsToModifier();
+
+	// this is kind of a hack, we are calling this for the side effect of forcing the modifier to recache volume shapes after we just invalidated them.
+	DaySequenceModifier->IsBlendTargetInAnyVolume();
+	
+#if WITH_EDITOR
+	const UWorld* World = GetWorld();
+	if (World && World->WorldType != EWorldType::Editor)
+	{
+#endif
+	SetupVolumeCallbacks();
+#if WITH_EDITOR
+	}
+#endif
+}
+
+void ADaySequenceModifierVolume::AddShapeComponentsToModifier()
+{
+	FComponentReference DefaultBoxReference;
+	DefaultBoxReference.OverrideComponent = DefaultBox;
+	DaySequenceModifier->AddVolumeShapeComponent(DefaultBoxReference);
+
+
+}
+
+void ADaySequenceModifierVolume::SetupVolumeCallbacks()
+{
+	if (!DaySequenceModifier->GetOnVolumeBlendTargetOverlapBegin().IsBoundToObject(this))
+	{
+		DaySequenceModifier->GetOnVolumeBlendTargetOverlapBegin().AddWeakLambda(this, [this](AActor*)
+		{
+			OnOverlapBegin();
+		});
+	}
+	
+	if (!DaySequenceModifier->GetOnVolumeBlendTargetOverlapEnd().IsBoundToObject(this))
+	{
+		DaySequenceModifier->GetOnVolumeBlendTargetOverlapEnd().AddWeakLambda(this, [this](AActor*)
+		{
+			OnOverlapEnd();
+		});
+	}
+}
+
+void ADaySequenceModifierVolume::OnOverlapBegin() const
+{
+	DaySequenceModifier->EnableModifier();
+}
+
+void ADaySequenceModifierVolume::OnOverlapEnd() const
+{
+	DaySequenceModifier->DisableModifier();
+}
+
+void ADaySequenceModifierVolume::PlayerControllerSetup()
+{
+#if WITH_EDITOR
+	UWorld* World = GetWorld();
+	if (World && World->WorldType != EWorldType::Editor)
+	{
+#endif
+	SetupBlendTargetCallbacks();
+	CachePlayerController();
+#if WITH_EDITOR
+	}
+#endif
+}
+
+void ADaySequenceModifierVolume::CachePlayerController()
+{
+	if (const UWorld* World = GetWorld())
+	{
+		for (FConstPlayerControllerIterator Iterator = World->GetPlayerControllerIterator(); Iterator; ++Iterator)
+		{
+			if (APlayerController* PlayerController = Iterator->Get())
+			{
+				if (PlayerController->IsLocalPlayerController())
+				{
+					CachedPlayerController = PlayerController;
+					break;
+				}
+			}
+		}
+	}
+
+	// If no local player controller found, queue this function for next tick.
+	if (!IsValid(CachedPlayerController))
+	{
+		QueuePlayerControllerQuery();
+	}
+	else
+	{
+		if (AActor* ViewTarget = CachedPlayerController->GetViewTarget())
+		{
+			SetBlendTarget(ViewTarget);
+		}
+	}
+}
+
+void ADaySequenceModifierVolume::QueuePlayerControllerQuery()
+{
+	if (!IsValid(this))
+	{
+		return;
+	}
+
+	if (const UWorld* World = GetWorld())
+	{
+		World->GetTimerManager().SetTimerForNextTick([this]()
+		{
+			CachePlayerController();
+		});
+	}
+}
+
+void ADaySequenceModifierVolume::SetupBlendTargetCallbacks()
+{
+	if (!ViewTargetChangedHandle.IsValid())
+	{
+		ViewTargetChangedHandle = FGameDelegates::Get().GetViewTargetChangedDelegate().AddWeakLambda(this, [this](APlayerController* PC, AActor* OldTarget, AActor* NewTarget)
+		{
+			if (!IsValid(CachedPlayerController))
+			{
+				QueuePlayerControllerQuery();
+				return;
+			}
+
+			if (PC != CachedPlayerController)
+			{
+				return;
+			}
+
+			// there is a lot of extra logic we can do here to stop popping (for example, try cast NewTarget to APlayerController which happens when view target unset).
+			// for now just set
+			SetBlendTarget(NewTarget);
+		});
+	}
+}
+
+void ADaySequenceModifierVolume::DaySequenceActorSetup()
+{
+	SetupDaySequenceSubsystemCallbacks();
+	BindToDaySequenceActor();
+}
+
+void ADaySequenceModifierVolume::BindToDaySequenceActor()
+{
+	bool bNewActorSet = false;
+	
+	if (const UWorld* World = GetWorld())
+	{
+		if (const UDaySequenceSubsystem* DaySequenceSubsystem = World->GetSubsystem<UDaySequenceSubsystem>())
+		{
+			if (ADaySequenceActor* NewActor = DaySequenceSubsystem->GetDaySequenceActor())
+			{
+				if (NewActor != DaySequenceActor)
+				{
+					DaySequenceActor = NewActor;
+					bNewActorSet = true;
+				}
+			}
+		}
+	}
+	
+	if (bNewActorSet)
+	{
+		DaySequenceModifier->BindToDaySequenceActor(DaySequenceActor);
+
+		VolumeSetup();
+		PlayerControllerSetup();
+
+		OnDaySequenceActorBound(DaySequenceActor);
+	}
+}
+
+void ADaySequenceModifierVolume::SetupDaySequenceSubsystemCallbacks()
+{
+	if (const UWorld* World = GetWorld())
+	{
+		if (UDaySequenceSubsystem* DaySequenceSubsystem = World->GetSubsystem<UDaySequenceSubsystem>())
+		{
+			// Prevent consecutive calls to this function from adding redundant lambdas to invocation list.
+			if (!DaySequenceSubsystem->OnDaySequenceActorSetEvent.IsBoundToObject(this))
+			{
+				DaySequenceSubsystem->OnDaySequenceActorSetEvent.AddWeakLambda(this, [this](ADaySequenceActor* InActor)
+				{
+					BindToDaySequenceActor();
+				});
+			}
+		}
+	}
+}
\ No newline at end of file
diff --git a/Engine/Plugins/Experimental/DaySequence/Source/DaySequence/Private/Actors/ProceduralDaySequenceActor.cpp b/Engine/Plugins/Experimental/DaySequence/Source/DaySequence/Private/Actors/ProceduralDaySequenceActor.cpp
new file mode 100644
--- /dev/null
+++ b/Engine/Plugins/Experimental/DaySequence/Source/DaySequence/Private/Actors/ProceduralDaySequenceActor.cpp
@@ -0,0 +1,47 @@
+﻿// Copyright Epic Games, Inc. All Rights Reserved.
+
+#include "Actors/ProceduralDaySequenceActor.h"
+
+#include "DaySequence.h"
+#include "DaySequenceCollectionAsset.h"
+#include "ProceduralDaySequenceBuilder.h"
+#include "Sections/MovieSceneSubSection.h"
+
+void AProceduralDaySequenceActor::OnConstruction(const FTransform& Transform)
+{
+	CreateProceduralSequence();
+	
+	Super::OnConstruction(Transform);
+}
+
+#if WITH_EDITOR
+void AProceduralDaySequenceActor::PostEditChangeProperty(FPropertyChangedEvent& PropertyChangedEvent)
+{
+	Super::PostEditChangeProperty(PropertyChangedEvent);
+}
+#endif
+
+void AProceduralDaySequenceActor::BuildSequence(UProceduralDaySequenceBuilder* SequenceBuilder) {}
+
+void AProceduralDaySequenceActor::InitializeDaySequences()
+{
+	Super::InitializeDaySequences();
+	
+	CreateProceduralSequence();
+
+	InitializeDaySequence(ProceduralSequence);
+}
+
+void AProceduralDaySequenceActor::CreateProceduralSequence()
+{
+	if (ProceduralSequence && !bProceduralSequenceInvalid)
+	{
+		return;
+	}
+	
+	UProceduralDaySequenceBuilder* SequenceBuilder = NewObject<UProceduralDaySequenceBuilder>();
+	ProceduralSequence = SequenceBuilder->Initialize(this, ProceduralSequence);
+	BuildSequence(SequenceBuilder);
+	
+	bProceduralSequenceInvalid = false;
+}
\ No newline at end of file
diff --git a/Engine/Plugins/Experimental/DaySequence/Source/DaySequence/Private/Actors/SimpleDaySequenceActor.cpp b/Engine/Plugins/Experimental/DaySequence/Source/DaySequence/Private/Actors/SimpleDaySequenceActor.cpp
new file mode 100644
--- /dev/null
+++ b/Engine/Plugins/Experimental/DaySequence/Source/DaySequence/Private/Actors/SimpleDaySequenceActor.cpp
@@ -0,0 +1,65 @@
+﻿// Copyright Epic Games, Inc. All Rights Reserved.
+
+#include "Actors/SimpleDaySequenceActor.h"
+
+#include "DaySequenceSubsystem.h"
+#include "Components/DirectionalLightComponent.h"
+#include "Components/ExponentialHeightFogComponent.h"
+#include "Components/SkyAtmosphereComponent.h"
+#include "Components/SkyLightComponent.h"
+#include "Components/VolumetricCloudComponent.h"
+#include "Engine/World.h"
+
+ASimpleDaySequenceActor::ASimpleDaySequenceActor(const FObjectInitializer& Init)
+: Super(Init)
+{
+	SunRootComponent = CreateDefaultSubobject<USceneComponent>(TEXT("SunRoot"));
+	SunRootComponent->SetupAttachment(RootComponent);
+	
+	SunComponent = CreateDefaultSubobject<UDirectionalLightComponent>(TEXT("Sun"));
+	SunComponent->SetupAttachment(SunRootComponent);
+	
+	ExponentialHeightFogComponent = CreateOptionalDefaultSubobject<UExponentialHeightFogComponent>(TEXT("ExponentialHeightFog"));
+	ExponentialHeightFogComponent->SetupAttachment(RootComponent);
+	ExponentialHeightFogComponent->bEnableVolumetricFog = true;
+
+	SkyAtmosphereComponent = CreateDefaultSubobject<USkyAtmosphereComponent>(TEXT("SkyAtmosphere"));
+	SkyAtmosphereComponent->SetupAttachment(RootComponent);
+
+	SkyLightComponent = CreateDefaultSubobject<USkyLightComponent>(TEXT("SkyLight"));
+	SkyLightComponent->SetupAttachment(RootComponent);
+	SkyLightComponent->bRealTimeCapture = true;
+	SkyLightComponent->bLowerHemisphereIsBlack = false;
+
+	VolumetricCloudComponent = CreateOptionalDefaultSubobject<UVolumetricCloudComponent>(TEXT("VolumetricCloud"));
+    VolumetricCloudComponent->SetupAttachment(RootComponent);
+}
+
+void ASimpleDaySequenceActor::BeginPlay()
+{
+	Super::BeginPlay();
+	
+	if (const UWorld* World = GetWorld())
+	{
+		if (UDaySequenceSubsystem* DaySequenceSubsystem = World->GetSubsystem<UDaySequenceSubsystem>())
+		{
+			DaySequenceSubsystem->SetDaySequenceActor(this);
+		}
+	}
+}
+
+void ASimpleDaySequenceActor::OnConstruction(const FTransform& Transform)
+{
+	Super::OnConstruction(Transform);
+	
+	if (const UWorld* World = GetWorld())
+	{
+		if (UDaySequenceSubsystem* DaySequenceSubsystem = World->GetSubsystem<UDaySequenceSubsystem>())
+		{
+			if (DaySequenceSubsystem->GetDaySequenceActor(/* bFindFallbackOnNull */ false) != this)
+			{
+				DaySequenceSubsystem->SetDaySequenceActor(this);
+			}
+		}
+	}
+}
\ No newline at end of file
diff --git a/Engine/Plugins/Experimental/DaySequence/Source/DaySequence/Private/Actors/SunAngleDaySequenceActor.cpp b/Engine/Plugins/Experimental/DaySequence/Source/DaySequence/Private/Actors/SunAngleDaySequenceActor.cpp
new file mode 100644
--- /dev/null
+++ b/Engine/Plugins/Experimental/DaySequence/Source/DaySequence/Private/Actors/SunAngleDaySequenceActor.cpp
@@ -0,0 +1,34 @@
+﻿// Copyright Epic Games, Inc. All Rights Reserved.
+
+#include "Actors/SunAngleDaySequenceActor.h"
+
+#include "Components/DirectionalLightComponent.h"
+#include "ProceduralDaySequenceBuilder.h"
+
+void ASunAngleDaySequenceActor::Tick(float DeltaTime)
+{
+	Super::Tick(DeltaTime);
+
+#if WITH_EDITOR
+	ApplySunAngle();
+#endif
+}
+
+void ASunAngleDaySequenceActor::ApplySunAngle()
+{
+	SunRootComponent->SetRelativeRotation(FRotator(0.0, SunYaw, SunRoll));
+	SunComponent->SetRelativeRotation(FRotator(SunPitch + 90, 0.0, 0.0));
+}
+
+void ASunAngleDaySequenceActor::BuildSequence(UProceduralDaySequenceBuilder* SequenceBuilder)
+{
+	SequenceBuilder->SetActiveBoundObject(this);
+
+	SequenceBuilder->AddScalarKey("SunPitch", 0.0, 0.0, RCIM_Linear);
+	SequenceBuilder->AddScalarKey("SunPitch", 1.0, 360.0, RCIM_Linear);
+}
+
+void ASunAngleDaySequenceActor::SequencePlayerUpdated(float CurrentTime, float PreviousTime)
+{
+	ApplySunAngle();
+}
\ No newline at end of file
diff --git a/Engine/Plugins/Experimental/DaySequence/Source/DaySequence/Private/Actors/SunPositionDaySequenceActor.cpp b/Engine/Plugins/Experimental/DaySequence/Source/DaySequence/Private/Actors/SunPositionDaySequenceActor.cpp
new file mode 100644
--- /dev/null
+++ b/Engine/Plugins/Experimental/DaySequence/Source/DaySequence/Private/Actors/SunPositionDaySequenceActor.cpp
@@ -0,0 +1,241 @@
+// Copyright Epic Games, Inc. All Rights Reserved.
+
+#include "Actors/SunPositionDaySequenceActor.h"
+
+#include "Components/DirectionalLightComponent.h"
+#include "ProceduralDaySequenceBuilder.h"
+
+namespace UE::DaySequence
+{
+	struct FSunPositionData
+	{
+		/** Sun Elevation */
+		double Elevation;
+
+		/** Sun Elevation, corrected for atmospheric diffraction */
+		double CorrectedElevation;
+
+		/** Sun azimuth */
+		double Azimuth;
+
+		/** Sunrise time */
+		FTimespan SunriseTime;
+
+		/** Sunset time */
+		FTimespan SunsetTime;
+
+		/** Solar noon */
+		FTimespan SolarNoon;
+	};
+
+	void GetSunPosition(double Latitude, double Longitude, double TimeZone, bool bIsDaylightSavingTime, int32 Year, int32 Month, int32 Day, int32 Hours, int32 Minutes, int32 Seconds, int32 Milliseconds, FSunPositionData& SunPositionData)
+	{
+		if (!FDateTime::Validate(Year, Month, Day, Hours, Minutes, Seconds, Milliseconds))
+		{
+			return;
+		}
+
+		FDateTime CalcTime(Year, Month, Day, Hours, Minutes, Seconds, Milliseconds);
+
+		double TimeOffset = TimeZone;
+		if (bIsDaylightSavingTime)
+		{
+			TimeOffset += 1.0;
+		}
+
+		double LatitudeRad = FMath::DegreesToRadians(Latitude);
+
+		// Get the julian day (number of days since Jan 1st of the year 4713 BC)
+		double JulianDay = CalcTime.GetJulianDay();
+		double JulianCentury = (JulianDay - 2451545.0) / 36525.0;
+
+		// Get the sun's mean longitude , referred to the mean equinox of julian date
+		double GeomMeanLongSunDeg = FMath::Fmod(280.46646 + JulianCentury*(36000.76983 + JulianCentury*0.0003032), 360.0);
+		double GeomMeanLongSunRad = FMath::DegreesToRadians(GeomMeanLongSunDeg);
+
+		// Get the sun's mean anomaly
+		double GeomMeanAnomSunDeg = 357.52911 + JulianCentury*(35999.05029 - 0.0001537*JulianCentury);
+		double GeomMeanAnomSunRad = FMath::DegreesToRadians(GeomMeanAnomSunDeg);
+
+		// Get the earth's orbit eccentricity
+		double EccentEarthOrbit = 0.016708634 - JulianCentury*(0.000042037 + 0.0000001267*JulianCentury);
+		
+		// Get the sun's equation of the center
+		double SunEqOfCtr = FMath::Sin(GeomMeanAnomSunRad)*(1.914602 - JulianCentury*(0.004817 + 0.000014*JulianCentury))
+						  + FMath::Sin(2.0 * GeomMeanAnomSunRad)*(0.019993 - 0.000101*JulianCentury)
+						  + FMath::Sin(3.0 * GeomMeanAnomSunRad)*0.000289;
+		
+		// Get the sun's true longitude
+		double SunTrueLongDeg = GeomMeanLongSunDeg + SunEqOfCtr;
+
+		// Get the sun's true anomaly
+		//	double SunTrueAnomDeg = GeomMeanAnomSunDeg + SunEqOfCtr;
+		//	double SunTrueAnomRad = FMath::DegreesToRadians(SunTrueAnomDeg);
+
+		// Get the earth's distance from the sun
+		//	double SunRadVectorAUs = (1.000001018*(1.0 - EccentEarthOrbit*EccentEarthOrbit)) / (1.0 + EccentEarthOrbit*FMath::Cos(SunTrueAnomRad));
+		
+		// Get the sun's apparent longitude
+		double SunAppLongDeg = SunTrueLongDeg - 0.00569 - 0.00478*FMath::Sin(FMath::DegreesToRadians(125.04 - 1934.136*JulianCentury));
+		double SunAppLongRad = FMath::DegreesToRadians(SunAppLongDeg);
+
+		// Get the earth's mean obliquity of the ecliptic
+		double MeanObliqEclipticDeg = 23.0 + (26.0 + ((21.448 - JulianCentury*(46.815 + JulianCentury*(0.00059 - JulianCentury*0.001813)))) / 60.0) / 60.0;
+		
+		// Get the oblique correction
+		double ObliqCorrDeg = MeanObliqEclipticDeg + 0.00256*FMath::Cos(FMath::DegreesToRadians(125.04 - 1934.136*JulianCentury));
+		double ObliqCorrRad = FMath::DegreesToRadians(ObliqCorrDeg);
+
+		// Get the sun's right ascension
+		double SunRtAscenRad = FMath::Atan2(FMath::Cos(ObliqCorrRad)*FMath::Sin(SunAppLongRad), FMath::Cos(SunAppLongRad));
+		double SunRtAscenDeg = FMath::RadiansToDegrees(SunRtAscenRad);
+		
+		// Get the sun's declination
+		double SunDeclinRad = FMath::Asin(FMath::Sin(ObliqCorrRad)*FMath::Sin(SunAppLongRad));
+		double SunDeclinDeg = FMath::RadiansToDegrees(SunDeclinRad);
+
+		double VarY = FMath::Pow(FMath::Tan(ObliqCorrRad / 2.0), 2.0);
+		
+		// Get the equation of time
+		double EqOfTimeMinutes = 4.0 * FMath::RadiansToDegrees(VarY*FMath::Sin(2.0 * GeomMeanLongSunRad) - 2.0 * EccentEarthOrbit*FMath::Sin(GeomMeanAnomSunRad) + 4.0 * EccentEarthOrbit*VarY*FMath::Sin(GeomMeanAnomSunRad)*FMath::Cos(2.0 * GeomMeanLongSunRad) - 0.5*VarY*VarY*FMath::Sin(4.0 * GeomMeanLongSunRad) - 1.25*EccentEarthOrbit*EccentEarthOrbit*FMath::Sin(2.0 * GeomMeanAnomSunRad));
+
+		// Get the hour angle of the sunrise
+		double HASunriseDeg = FMath::RadiansToDegrees(FMath::Acos(FMath::Cos(FMath::DegreesToRadians(90.833)) / (FMath::Cos(LatitudeRad)*FMath::Cos(SunDeclinRad)) - FMath::Tan(LatitudeRad)*FMath::Tan(SunDeclinRad)));
+		//	double SunlightDurationMinutes = 8.0 * HASunriseDeg;
+
+		// Get the local time of the sun's rise and set
+		double SolarNoonLST = (720.0 - 4.0 * Longitude - EqOfTimeMinutes + TimeOffset * 60.0) / 1440.0;
+		double SunriseTimeLST = SolarNoonLST - HASunriseDeg * 4.0 / 1440.0;
+		double SunsetTimeLST = SolarNoonLST + HASunriseDeg * 4.0 / 1440.0;
+
+		// Get the true solar time
+		double TrueSolarTimeMinutes = FMath::Fmod(CalcTime.GetTimeOfDay().GetTotalMinutes() + EqOfTimeMinutes + 4.0 * Longitude - 60.0 * TimeOffset, 1440.0);
+
+		// Get the hour angle of current time
+		double HourAngleDeg = TrueSolarTimeMinutes < 0 ? TrueSolarTimeMinutes / 4.0 + 180 : TrueSolarTimeMinutes / 4.0 - 180.0;
+		double HourAngleRad = FMath::DegreesToRadians(HourAngleDeg);
+
+		// Get the solar zenith angle
+		double SolarZenithAngleRad = FMath::Acos(FMath::Sin(LatitudeRad)*FMath::Sin(SunDeclinRad) + FMath::Cos(LatitudeRad)*FMath::Cos(SunDeclinRad)*FMath::Cos(HourAngleRad));
+		double SolarZenithAngleDeg = FMath::RadiansToDegrees(SolarZenithAngleRad);
+
+		// Get the sun elevation
+		double SolarElevationAngleDeg = 90.0 - SolarZenithAngleDeg;
+		double SolarElevationAngleRad = FMath::DegreesToRadians(SolarElevationAngleDeg);
+		double TanOfSolarElevationAngle = FMath::Tan(SolarElevationAngleRad);
+
+		// Get the approximated atmospheric refraction
+		double ApproxAtmosphericRefractionDeg = 0.0;
+		if (SolarElevationAngleDeg <= 85.0)
+		{
+			if (SolarElevationAngleDeg > 5.0)
+			{
+				ApproxAtmosphericRefractionDeg = 58.1 / TanOfSolarElevationAngle - 0.07 / FMath::Pow(TanOfSolarElevationAngle, 3) + 0.000086 / FMath::Pow(TanOfSolarElevationAngle, 5) / 3600.0;
+			}
+			else
+			{
+				if (SolarElevationAngleDeg > -0.575)
+				{
+					ApproxAtmosphericRefractionDeg = 1735.0 + SolarElevationAngleDeg * (-518.2 + SolarElevationAngleDeg * (103.4 + SolarElevationAngleDeg * (-12.79 + SolarElevationAngleDeg * 0.711)));
+				}
+				else
+				{
+					ApproxAtmosphericRefractionDeg = -20.772 / TanOfSolarElevationAngle;
+				}
+			}
+			ApproxAtmosphericRefractionDeg /= 3600.0;
+		}
+
+		// Get the corrected solar elevation
+		double SolarElevationcorrectedforatmrefractionDeg = SolarElevationAngleDeg + ApproxAtmosphericRefractionDeg;
+
+		// Get the solar azimuth 
+		double tmp = FMath::RadiansToDegrees(FMath::Acos(((FMath::Sin(LatitudeRad)*FMath::Cos(SolarZenithAngleRad)) - FMath::Sin(SunDeclinRad)) / (FMath::Cos(LatitudeRad)*FMath::Sin(SolarZenithAngleRad))));
+		double SolarAzimuthAngleDegcwfromN = HourAngleDeg > 0.0 ? FMath::Fmod(tmp + 180.0, 360.0) : FMath::Fmod(540.0 - tmp, 360.0);
+
+
+		// offset elevation angle to fit with UE coords system
+		SunPositionData.Elevation = 180.0 + SolarElevationAngleDeg;
+		SunPositionData.CorrectedElevation = 180.0 + SolarElevationcorrectedforatmrefractionDeg;
+		SunPositionData.Azimuth = SolarAzimuthAngleDegcwfromN;
+		SunPositionData.SolarNoon = FTimespan::FromDays(SolarNoonLST);
+		SunPositionData.SunriseTime = FTimespan::FromDays(SunriseTimeLST);
+		SunPositionData.SunsetTime = FTimespan::FromDays(SunsetTimeLST);
+	}
+}
+
+void ASunPositionDaySequenceActor::Tick(float DeltaTime)
+{
+	Super::Tick(DeltaTime);
+
+#if WITH_EDITOR
+	ApplySunAngle();
+#endif
+}
+
+void ASunPositionDaySequenceActor::ApplySunAngle()
+{
+	SunComponent->SetRelativeRotation(FRotator(SunPitch, SunYaw, 0.0));
+
+	Azimuth = SunYaw;
+	Elevation = SunPitch - 180.0;
+}
+
+void ASunPositionDaySequenceActor::BuildSequence(UProceduralDaySequenceBuilder* SequenceBuilder)
+{
+	using namespace UE::DaySequence;
+	
+	SequenceBuilder->SetActiveBoundObject(this);
+	
+	const double NormalizedTimeIncrement = 1.0 / FMath::Max(KeyCount - 1, static_cast<unsigned>(1));
+	for (unsigned int Key = 0; Key < KeyCount; ++Key)
+	{
+		const double KeyTime = Key * NormalizedTimeIncrement;
+		
+		FSunPositionData Data;
+
+		double Hours = (double)Key / (double)KeyCount * 24.0;
+		int32 IntHours = Hours;
+
+		double Minutes = FMath::Frac(Hours) * 60.0;
+		int32 IntMinutes = Minutes;
+
+		double Seconds = FMath::Frac(Minutes) * 60.0;
+		int32 IntSeconds = Seconds;
+
+		double Milliseconds = FMath::Frac(Seconds) * 1000.0;
+		int32 IntMilliseconds = Milliseconds;
+		
+		GetSunPosition(Latitude, Longitude, TimeZone, bIsDaylightSavings, Time.GetYear(), Time.GetMonth(), Time.GetDay(), IntHours, IntMinutes, IntSeconds, IntMilliseconds, Data);
+		
+		SequenceBuilder->AddScalarKey("SunPitch", KeyTime, Data.CorrectedElevation);
+		SequenceBuilder->AddScalarKey("SunYaw", KeyTime, Data.Azimuth);
+	}
+}
+
+void ASunPositionDaySequenceActor::SequencePlayerUpdated(float CurrentTime, float PreviousTime)
+{
+	ApplySunAngle();
+}
+
+#if WITH_EDITOR
+void ASunPositionDaySequenceActor::PostEditChangeProperty(FPropertyChangedEvent& PropertyChangedEvent)
+{
+	const FName PropertyName = PropertyChangedEvent.Property ? PropertyChangedEvent.Property->GetFName() : NAME_None;
+
+	// When we finish a drag setting a property that changes the procedural sequence, mark the sequence invalid.
+	// We don't do this during the drag because it will do a _lot_ of work.
+	if (PropertyChangedEvent.ChangeType == EPropertyChangeType::ValueSet &&
+		(PropertyName == GET_MEMBER_NAME_CHECKED(ASunPositionDaySequenceActor, KeyCount) ||
+		 PropertyName == GET_MEMBER_NAME_CHECKED(ASunPositionDaySequenceActor, Time) ||
+		 PropertyName == GET_MEMBER_NAME_CHECKED(ASunPositionDaySequenceActor, TimeZone) ||
+		 PropertyName == GET_MEMBER_NAME_CHECKED(ASunPositionDaySequenceActor, Latitude) ||
+		 PropertyName == GET_MEMBER_NAME_CHECKED(ASunPositionDaySequenceActor, Longitude) ||
+		 PropertyName == GET_MEMBER_NAME_CHECKED(ASunPositionDaySequenceActor, bIsDaylightSavings)))
+	{
+		bProceduralSequenceInvalid = true;
+	}
+
+	Super::PostEditChangeProperty(PropertyChangedEvent);
+}
+#endif
\ No newline at end of file
diff --git a/Engine/Plugins/Experimental/DaySequence/Source/DaySequence/Private/DaySequenceActor.cpp b/Engine/Plugins/Experimental/DaySequence/Source/DaySequence/Private/DaySequenceActor.cpp
new file mode 100644
--- /dev/null
+++ b/Engine/Plugins/Experimental/DaySequence/Source/DaySequence/Private/DaySequenceActor.cpp
@@ -0,0 +1,1635 @@
+// Copyright Epic Games, Inc. All Rights Reserved.
+
+#include "DaySequenceActor.h"
+#include "DaySequenceCollectionAsset.h"
+#include "DaySequenceHelpers.h"
+#include "DaySequence.h"
+#include "DaySequenceModule.h"
+#include "DaySequencePlayer.h"
+#include "DaySequenceTrack.h"
+#include "DaySequenceConditionTag.h"
+#include "DaySequenceSubsystem.h"
+
+#include "Components/BillboardComponent.h"
+#include "Curves/CurveFloat.h"
+#include "Engine/Texture2D.h"
+#include "LevelSequenceActor.h" // For FBoundActorProxy
+#include "MovieScene.h"
+#include "MovieSceneTimeHelpers.h"
+#include "MovieSceneBindingOverrides.h"
+#include "Net/UnrealNetwork.h"
+#include "MovieSceneTimeController.h"
+#include "Sections/MovieSceneSubSection.h"
+#include "Tracks/MovieSceneSubTrack.h"
+#include "UObject/ConstructorHelpers.h"
+#include "UObject/UObjectThreadContext.h" // For FUObjectThreadContext
+#include "ProfilingDebugging/CsvProfiler.h"
+
+#if ENABLE_DRAW_DEBUG
+#include "Engine/Engine.h"	// GEngine
+#include "GameFramework/HUD.h"	// AHUD::OnShowDebugInfo
+#include "Engine/Canvas.h"	// UCanvas
+#endif
+
+#include <limits>
+
+#if WITH_EDITOR
+#include "PropertyHandle.h"
+#endif
+
+#include UE_INLINE_GENERATED_CPP_BY_NAME(DaySequenceActor)
+
+namespace UE::DaySequence
+{
+#if ENABLE_DRAW_DEBUG
+	FDaySequenceDebugEntry::FDaySequenceDebugEntry(FShowDebugDataConditionFunction InShowCondition, FGetDebugDataFunction InGetData)
+	: ShowCondition(InShowCondition)
+	, GetData(InGetData)
+	{}
+	
+	int32 GDaySequenceDebugLevel = 0;
+	FAutoConsoleVariableRef CVarDaySequenceActorDebugLevel(
+		TEXT("DaySequence.DebugLevel"),
+		GDaySequenceDebugLevel,
+		TEXT("The debug level to use.")
+	);
+#endif
+	
+	int32 GFrameBudgetMicroseconds = 30;
+	FAutoConsoleVariableRef CVarFrameBudgetMicroseconds(
+		TEXT("TimeOfDay.FrameBudget"),
+		GFrameBudgetMicroseconds,
+		TEXT("(Default: 30us) Approximate max per-frame budget for time-of-day actors in microseconds.")
+	);
+
+
+	struct FStaticTimeControllerOverride : FMovieSceneTimeController
+	{
+		TSharedPtr<FMovieSceneTimeController> PreviousTimeController;
+		float TimeInHours;
+		FFrameTime StaticTime;
+#if WITH_EDITORONLY_DATA
+		FTimecode PreviousTimeOfDayPreview;
+#endif
+
+		FStaticTimeControllerOverride(FFrameTime InStaticTime, TSharedPtr<FMovieSceneTimeController> InPreviousTimeController)
+			: PreviousTimeController(InPreviousTimeController)
+			, StaticTime(InStaticTime)
+		{}
+
+		void OnTick(float DeltaSeconds, float InPlayRate) override
+		{
+			if (PreviousTimeController)
+			{
+				PreviousTimeController->Tick(DeltaSeconds, InPlayRate);
+			}
+		}
+		void OnStartPlaying(const FQualifiedFrameTime& InStartTime) override
+		{
+			if (PreviousTimeController)
+			{
+				PreviousTimeController->StartPlaying(InStartTime);
+			}
+		}
+		void OnStopPlaying(const FQualifiedFrameTime& InStopTime) override
+		{
+			if (PreviousTimeController)
+			{
+				PreviousTimeController->StopPlaying(InStopTime);
+			}
+		}
+		FFrameTime OnRequestCurrentTime(const FQualifiedFrameTime& InCurrentTime, float InPlayRate)
+		{
+			return StaticTime;
+		}
+	};
+
+} // namespace UE::DaySequence
+
+ADaySequenceActor::ADaySequenceActor(const FObjectInitializer& Init)
+: Super(Init)
+, bRunDayCycle(true)
+, bUseInterpCurve(false)
+, DayLength(24, 0, 0, 0, false)
+, TimePerCycle(0, 5, 0, 0, false)
+, InitialTimeOfDay(6, 0, 0, 0, false)
+{
+	USceneComponent* SceneRootComponent = CreateDefaultSubobject<USceneComponent>(USceneComponent::GetDefaultSceneRootVariableName());
+	SetRootComponent(SceneRootComponent);
+
+#if WITH_EDITORONLY_DATA
+	if (!IsRunningCommandlet())
+	{
+		// Structure to hold one-time initialization
+		struct FConstructorStatics
+		{
+			ConstructorHelpers::FObjectFinderOptional<UTexture2D> DecalTexture;
+			FConstructorStatics() : DecalTexture(TEXT("/Engine/EditorResources/S_LevelSequence")) {}
+		};
+		static FConstructorStatics ConstructorStatics;
+
+		if (UBillboardComponent* LocalSpriteComponent = GetSpriteComponent())
+		{
+			LocalSpriteComponent->Sprite = ConstructorStatics.DecalTexture.Get();
+			LocalSpriteComponent->SetupAttachment(RootComponent);
+			LocalSpriteComponent->SetUsingAbsoluteScale(true);
+			LocalSpriteComponent->bReceivesDecals = false;
+			LocalSpriteComponent->bHiddenInGame = true;
+		}
+	}
+
+	TimeOfDayPreview = FTimecode(6, 0, 0, 0, false);
+#endif // WITH_EDITORONLY_DATA
+
+#if WITH_EDITOR || ENABLE_DRAW_DEBUG
+	// The DaySequenceActor is ticked separately in LevelTick. However, in Editor, we tick to allow
+	// deferred initialization of the root sequence outside of actor construction / BP reinstancing.
+	// Comment from Nick: We also now tick in dev builds + Editor to catch changes to
+	// GDaySequenceDebugLevel which is set via a CVar.
+	PrimaryActorTick.bCanEverTick = true;
+#endif
+	
+	// SequencePlayer must be a default sub object for it to be replicated correctly
+	SequencePlayer = Init.CreateDefaultSubobject<UDaySequencePlayer>(this, "AnimationPlayer", true);
+	BindingOverrides = Init.CreateDefaultSubobject<UMovieSceneBindingOverrides>(this, "BindingOverrides");
+
+	bAlwaysRelevant = true;
+	bReplicates = true;
+	bReplicateUsingRegisteredSubObjectList = true;
+	bReplicatePlayback = false;
+	SetHidden(false);
+
+#if ENABLE_DRAW_DEBUG
+	if (!IsTemplate())
+	{
+		AHUD::OnShowDebugInfo.AddUObject(this, &ADaySequenceActor::OnShowDebugInfo);
+	}
+#endif
+}
+
+void ADaySequenceActor::PostInitializeComponents()
+{
+	Super::PostInitializeComponents();
+
+	if (HasAuthority())
+	{
+		SetReplicates(bReplicatePlayback);
+	}
+
+	if (IsValid(SequencePlayer))
+	{
+		// Initialize this player for tick as soon as possible to ensure that a persistent
+		// reference to the tick manager is maintained
+		SequencePlayer->InitializeForTick(this);
+		SequencePlayer->OnSequenceUpdated().AddUObject(this, &ADaySequenceActor::OnSequencePlayerUpdate);
+	}
+	
+	InitializePlayer();
+}
+
+#if WITH_EDITOR || ENABLE_DRAW_DEBUG
+void ADaySequenceActor::OnConstruction(const FTransform& Transform)
+{
+#if WITH_EDITOR
+	// It is unsafe to update the root sequence (incl. its delegates)
+	// during actor construction. Defer to the next tick.
+	bUpdateRootSequenceOnTick = true;
+#endif
+}
+
+void ADaySequenceActor::Tick(float DeltaTime)
+{
+	check(!HasAnyFlags(RF_ClassDefaultObject));
+	Super::Tick(DeltaTime);
+	
+#if WITH_EDITOR
+	if (bUpdateRootSequenceOnTick && GetWorld() != nullptr && GetWorld()->WorldType == EWorldType::Editor)
+	{
+		UpdateRootSequence();
+		bUpdateRootSequenceOnTick = false;
+	}
+#endif
+
+	if (UE::DaySequence::GDaySequenceDebugLevel != CachedDebugLevel)
+	{
+		OnDebugLevelChanged.Broadcast(UE::DaySequence::GDaySequenceDebugLevel);
+		CachedDebugLevel = UE::DaySequence::GDaySequenceDebugLevel;
+	}
+}
+
+bool ADaySequenceActor::ShouldTickIfViewportsOnly() const
+{
+	return GetWorld() != nullptr && GetWorld()->WorldType == EWorldType::Editor;
+}
+#endif
+
+UDaySequencePlayer* ADaySequenceActor::GetSequencePlayer() const
+{
+	return SequencePlayer && SequencePlayer->IsValid() ? SequencePlayer : nullptr;
+}
+
+bool ADaySequenceActor::ContainsDaySequence(const UDaySequence* InDaySequence)
+{
+	if (InDaySequence && DaySequenceCollection)
+	{
+		for (const FDaySequenceCollectionEntry& SequenceAsset : DaySequenceCollection->DaySequences)
+		{
+			if (SequenceAsset.Sequence == InDaySequence)
+			{
+				return true;
+			}
+		}
+	}
+	return false;
+}
+
+void ADaySequenceActor::SetReplicatePlayback(bool bInReplicatePlayback)
+{
+	bReplicatePlayback = bInReplicatePlayback;
+	SetReplicates(bReplicatePlayback);
+}
+
+void ADaySequenceActor::GetLifetimeReplicatedProps(TArray<FLifetimeProperty>& OutLifetimeProps) const
+{
+	Super::GetLifetimeReplicatedProps(OutLifetimeProps);
+
+	DOREPLIFETIME(ADaySequenceActor, SequencePlayer);
+	DOREPLIFETIME(ADaySequenceActor, DayInterpCurve);
+	DOREPLIFETIME(ADaySequenceActor, bUseInterpCurve);
+}
+
+void ADaySequenceActor::PostInitProperties()
+{
+	Super::PostInitProperties();
+
+	// This is only checking that the UObject is valid, not that the player has been initialized yet.
+	if (IsValid(SequencePlayer))
+	{
+		// Have to initialize this here as any properties set on default subobjects inside the constructor
+		// get stomped by the CDO's properties when the constructor exits.
+		SequencePlayer->SetPlaybackClient(this);
+	}
+}
+
+void ADaySequenceActor::PostLoad()
+{
+	Super::PostLoad();
+
+#if WITH_EDITOR
+	// Initialize our transient preview time to InitialTimeOfDay on load.
+	TimeOfDayPreview = InitialTimeOfDay;
+	
+	// Build our root sequence after load to ensure that the editor can parse the root sequence
+	// hierarchy for editing binding overrides. This is only necessary for editor, since the
+	// root sequence will be initialized in PostInitializeComponents()/InitializePlayer() for runtime.
+	InitializeRootSequence();
+#endif
+
+#if WITH_EDITORONLY_DATA
+	// Fix sprite component so that it's attached to the root component. In the past, the sprite component was the root component.
+	UBillboardComponent* LocalSpriteComponent = FindComponentByClass<UBillboardComponent>();
+	if (LocalSpriteComponent && LocalSpriteComponent->GetAttachParent() != RootComponent)
+	{
+		LocalSpriteComponent->SetupAttachment(RootComponent);
+	}
+#endif
+}
+
+void ADaySequenceActor::BeginPlay()
+{
+	Super::BeginPlay();
+
+	if (UDaySequencePlayer* Player = GetSequencePlayer())
+	{
+		AddReplicatedSubObject(Player);
+
+		// Only play if we have a valid day sequence.
+		if (HasValidRootSequence())
+		{
+			// Always play a valid day sequence. Pause if bRunDayCycle is false
+			// to allow sequence spawnables and property tracks to be set from initial
+			// time of day.
+			Player->PlayLooping();
+			
+#if WITH_EDITORONLY_DATA
+			const bool bPause = !bRunDayCycle || bOverrideRunDayCycle;
+#else
+			const bool bPause = !bRunDayCycle;
+#endif
+			
+			if (bPause)
+			{
+				Player->Pause();
+			}
+		}
+	}
+}
+
+void ADaySequenceActor::EndPlay(const EEndPlayReason::Type EndPlayReason)
+{
+	if (UDaySequencePlayer* Player = GetSequencePlayer())
+	{
+		RemoveReplicatedSubObject(Player);
+
+		// Stop may modify a lot of actor state so it needs to be called
+		// during EndPlay (when Actors + World are still valid) instead
+		// of waiting for the UObject to be destroyed by GC.
+		Player->Stop();
+
+		Player->OnPlay.RemoveAll(this);
+		Player->OnPlayReverse.RemoveAll(this);
+		Player->OnStop.RemoveAll(this);
+
+		Player->TearDown();
+	}
+
+	Super::EndPlay(EndPlayReason);
+}
+
+void ADaySequenceActor::RewindForReplay()
+{
+	if (UDaySequencePlayer* Player = GetSequencePlayer())
+	{
+		Player->RewindForReplay();
+	}
+}
+
+void ADaySequenceActor::Destroyed()
+{
+#if WITH_EDITOR
+	if (const UWorld* World = GetWorld())
+	{
+		if (UDaySequenceSubsystem* DaySubsystem = World->GetSubsystem<UDaySequenceSubsystem>())
+		{
+			if (DaySubsystem->GetDaySequenceActor() == this)
+			{
+				DaySubsystem->SetDaySequenceActor(nullptr);
+			}
+		}
+	}
+#endif
+
+	Super::Destroyed();
+}
+
+bool ADaySequenceActor::RetrieveBindingOverrides(const FGuid& InBindingId, FMovieSceneSequenceID InSequenceID, TArray<UObject*, TInlineAllocator<1>>& OutObjects) const
+{
+	return BindingOverrides->LocateBoundObjects(InBindingId, InSequenceID, OutObjects);
+}
+
+UObject* ADaySequenceActor::GetInstanceData() const
+{
+	return nullptr;
+}
+
+bool ADaySequenceActor::GetIsReplicatedPlayback() const
+{
+	return bReplicatePlayback;
+}
+
+#if WITH_EDITOR
+TSharedPtr<FStructOnScope> ADaySequenceActor::GetObjectPickerProxy(TSharedPtr<IPropertyHandle> ObjectPropertyHandle)
+{
+	TSharedRef<FStructOnScope> Struct = MakeShared<FStructOnScope>(FBoundActorProxy::StaticStruct());
+	reinterpret_cast<FBoundActorProxy*>(Struct->GetStructMemory())->Initialize(ObjectPropertyHandle);
+	return Struct;
+}
+
+void ADaySequenceActor::UpdateObjectFromProxy(FStructOnScope& Proxy, IPropertyHandle& ObjectPropertyHandle)
+{
+	UObject* BoundActor = reinterpret_cast<FBoundActorProxy*>(Proxy.GetStructMemory())->BoundActor;
+	ObjectPropertyHandle.SetValue(BoundActor);
+}
+
+UMovieSceneSequence* ADaySequenceActor::RetrieveOwnedSequence() const
+{
+	return RootSequence;
+}
+
+bool ADaySequenceActor::GetReferencedContentObjects(TArray<UObject*>& Objects) const
+{
+	if (DaySequenceCollection)
+	{
+		Objects.Add(DaySequenceCollection);
+	}
+	Super::GetReferencedContentObjects(Objects);
+	return true;
+}
+
+void ADaySequenceActor::PostEditChangeProperty(FPropertyChangedEvent& PropertyChangedEvent)
+{
+	const FName PropertyName = (PropertyChangedEvent.Property != nullptr) ? PropertyChangedEvent.Property->GetFName() : NAME_None;
+	if (PropertyName == GET_MEMBER_NAME_CHECKED(ADaySequenceActor, TimeOfDayPreview))
+	{
+		// Force the change to ensure it is broadcast to clients.
+		SetTimeOfDayPreview(GetTimeOfDayPreview());
+	}
+	else if (PropertyName == GET_MEMBER_NAME_CHECKED(ADaySequenceActor, DayLength))
+	{
+		SetDayLength(GetDayLength());
+	}
+	else if (PropertyName == GET_MEMBER_NAME_CHECKED(ADaySequenceActor, TimePerCycle))
+	{
+		SetTimePerCycle(GetTimePerCycle());
+
+		// Update our root sequence with the new asset.
+		UpdateRootSequence();
+	}
+	else if (PropertyName == GET_MEMBER_NAME_CHECKED(ADaySequenceActor, InitialTimeOfDay))
+	{
+		SetInitialTimeOfDay(GetInitialTimeOfDay());
+	}
+	else if (PropertyName == GET_MEMBER_NAME_CHECKED(ADaySequenceActor, DaySequenceCollection))
+	{
+		SubSections.Empty();
+	}
+
+	Super::PostEditChangeProperty(PropertyChangedEvent);
+}
+
+void ADaySequenceActor::ConditionalSetTimeOfDayPreview(float InHours)
+{
+	// Wrap the input hours using day length
+	InHours = FMath::Frac(InHours / GetDayLength()) * GetDayLength(); 
+	const TOptional<FFrameRate> FrameRate = RootSequence ? RootSequence->GetMovieScene()->GetDisplayRate() : TOptional<FFrameRate>();
+	const FTimecode NewTimeOfDayPreview = DaySequenceHelpers::HoursToTimecode(InHours, FrameRate);
+	if (NewTimeOfDayPreview != TimeOfDayPreview)
+	{
+		TimeOfDayPreview = NewTimeOfDayPreview;
+		OnTimeOfDayPreviewChangedEvent.Broadcast(InHours);
+		OnTimeOfDayPreviewChanged.Broadcast(InHours);
+	}
+}
+
+#endif //WITH_EDITOR
+
+float ADaySequenceActor::GetTimeOfDayPreview() const
+{
+#if WITH_EDITOR
+	const TOptional<FFrameRate> FrameRate = RootSequence ? RootSequence->GetMovieScene()->GetDisplayRate() : TOptional<FFrameRate>();
+	return DaySequenceHelpers::TimecodeToHours(TimeOfDayPreview, FrameRate);
+#else
+	return 0.0f;
+#endif
+}
+
+void ADaySequenceActor::SetTimeOfDayPreview(float InHours)
+{
+#if WITH_EDITOR
+	// Wrap the input hours using day length
+	InHours = FMath::Frac(InHours / GetDayLength()) * GetDayLength(); 
+	const TOptional<FFrameRate> FrameRate = RootSequence ? RootSequence->GetMovieScene()->GetDisplayRate() : TOptional<FFrameRate>();
+	const FTimecode NewTimeOfDayPreview = DaySequenceHelpers::HoursToTimecode(InHours, FrameRate);
+	TimeOfDayPreview = NewTimeOfDayPreview;
+	OnTimeOfDayPreviewChangedEvent.Broadcast(InHours);
+	OnTimeOfDayPreviewChanged.Broadcast(InHours);
+#endif
+}
+
+void ADaySequenceActor::InitializePlayer()
+{
+	using namespace UE::DaySequence;
+
+	InitializeRootSequence();
+
+	if (GetWorld()->IsGameWorld())
+	{
+		// Initializing the sequence player may (will) overwrite the TimeController, so make sure we
+		// maintain the overridden one if SetStaticTimeOfDay was called before Initialize
+		TSharedPtr<FStaticTimeControllerOverride> OldTimeControllerOverride = WeakTimeControllerOverride.Pin();
+
+		if (SequencePlayer)
+		{
+			SequencePlayer->Initialize(RootSequence, this, GetPlaybackSettings(RootSequence));
+		}
+
+		UDaySequencePlayer* Player = GetSequencePlayer();
+		if (OldTimeControllerOverride && Player && OldTimeControllerOverride != Player->GetTimeController())
+		{
+			ensure(OldTimeControllerOverride->PreviousTimeController == nullptr);
+
+			OldTimeControllerOverride->PreviousTimeController = Player->GetTimeController();
+			Player->SetTimeControllerDirectly(OldTimeControllerOverride);
+		}
+	}
+}
+
+void ADaySequenceActor::InitializeRootSequence()
+{
+	CSV_SCOPED_TIMING_STAT_EXCLUSIVE(DaySequenceActor_InitializeRootSequence);
+	
+	ensureMsgf(!SequencePlayer || !SequencePlayer->IsValid(), TEXT("InitializeRootSequence called but the sequence player has already been initialized."));
+	
+#if WITH_EDITOR
+	// Do not generate the root sequence during cook.
+	if (IsRunningCookCommandlet())
+	{
+		return;
+	}
+#endif
+
+	OnPreRootSequenceChanged.Broadcast();
+
+#if ROOT_SEQUENCE_RECONSTRUCTION_ENABLED
+	UWorld* World = GetWorld();
+	const bool bInEditorWorld = World && World->WorldType == EWorldType::Editor;
+	
+	if (!RootSequence || !bInEditorWorld)
+	{
+#endif
+		RootSequence = NewObject<UDaySequence>(this, NAME_None, RF_Transient);
+		RootSequence->Initialize(RF_Transient);
+		RootSequence->SetSequenceFlags(EMovieSceneSequenceFlags::Volatile);
+
+		UMovieScene* RootMovieScene = RootSequence->GetMovieScene();
+		const FFrameRate DisplayRate = RootMovieScene->GetDisplayRate();
+		const float DaySeconds = DaySequenceHelpers::TimecodeToSeconds(TimePerCycle, DisplayRate);
+		const int32 RootDuration = RootMovieScene->GetTickResolution().AsFrameNumber(DaySeconds).Value;
+		RootMovieScene->SetPlaybackRange(0, RootDuration);
+#if WITH_EDITOR
+		RootMovieScene->SetPlaybackRangeLocked(true);
+#endif
+#if ROOT_SEQUENCE_RECONSTRUCTION_ENABLED
+	}
+#endif
+	
+	if (GetNetMode() == NM_DedicatedServer)
+	{
+		return;
+	}
+
+#if ENABLE_DRAW_DEBUG
+	for (TSharedPtr<UE::DaySequence::FDaySequenceDebugEntry> Entry : SubSectionDebugEntries)
+	{
+		UnregisterDebugEntry(Entry, ShowDebug_SubSequenceCategory);
+	}
+	
+	SubSectionDebugEntries.Empty();
+#endif
+	
+#if ROOT_SEQUENCE_RECONSTRUCTION_ENABLED
+	if (bInEditorWorld)
+	{
+		FSubSectionPreserveMap SectionsToPreserve;
+		for (const UMovieSceneTrack* Track : RootSequence->GetMovieScene()->GetTracks())
+		{
+			if (const TArray<UMovieSceneSection*>& Sections = Track->GetAllSections(); Sections.Num() > 0)
+			{
+				// There is an assumption of 1 section per track. If that assumption is not respected, root sequence reconstruction will likely break.
+				if (UMovieSceneSubSection* SubSection = Cast<UMovieSceneSubSection>(Sections[0]))
+				{
+					SectionsToPreserve.Add(SubSection, false);
+				}
+			}
+		}
+		
+		// This returns true if we need to do a full reinit and false if we can skip init
+		if (MarkDaySequences(&SectionsToPreserve))
+		{
+			InitializeDaySequences();
+		}
+
+		OnPostInitializeDaySequences.Broadcast(&SectionsToPreserve);
+
+		// Remove all unmarked sections
+		for (const TPair<UMovieSceneSubSection*, bool>& Section : SectionsToPreserve)
+		{
+			if (!Section.Value)
+			{
+				UMovieSceneTrack* Track = Section.Key->GetTypedOuter<UMovieSceneTrack>();
+				UMovieScene* MovieScene = Track->GetTypedOuter<UMovieScene>();
+	
+				check(Track && MovieScene);
+	
+				MovieScene->RemoveTrack(*Track);
+				MovieScene->MarkAsChanged();
+			}
+		}
+	}
+	else
+	{
+		InitializeDaySequences();
+		OnPostInitializeDaySequences.Broadcast(nullptr);
+	}
+#else
+	InitializeDaySequences();
+	OnPostInitializeDaySequences.Broadcast(nullptr);
+#endif
+
+#if ENABLE_DRAW_DEBUG
+	if (!IsDebugCategoryRegistered(ShowDebug_SubSequenceCategory))
+	{
+		RegisterDebugCategory(ShowDebug_SubSequenceCategory,  OnShowDebugInfoDrawFunction);
+	}
+
+	for (TSharedPtr<UE::DaySequence::FDaySequenceDebugEntry> Entry : SubSectionDebugEntries)
+	{
+		RegisterDebugEntry(Entry, ShowDebug_SubSequenceCategory);
+	}
+#endif
+	
+	OnPostRootSequenceChanged.Broadcast();
+}
+
+#if ROOT_SEQUENCE_RECONSTRUCTION_ENABLED
+bool ADaySequenceActor::MarkDaySequences(FSubSectionPreserveMap* SectionsToPreserve)
+{
+	bool bReinit = true;
+
+	if (SectionsToPreserve)
+	{
+		// Mark all subsections we have recorded for keep in the root sequence
+		// This is a fast path we take only if all of our subsections are in the root sequence
+		for (TWeakObjectPtr<UMovieSceneSubSection> SubSection : SubSections)
+		{
+			if (const UMovieSceneSubSection* StrongSubSection = SubSection.Get())
+			{
+				if (bool* SectionToPreserveFlag = SectionsToPreserve->Find(StrongSubSection))
+				{
+					*SectionToPreserveFlag = true;
+					bReinit = false;
+				}
+				else
+				{
+					// If we have a subsection that is not in the root sequence, break and reinit completely
+					bReinit = true;
+					break;
+				}
+			}
+		}
+	}
+
+	if (bReinit)
+	{
+		if (SectionsToPreserve)
+		{
+			// Mark all sections associated with this modifier for delete before we do a full reinit
+			for (TWeakObjectPtr<UMovieSceneSubSection> SubSection : SubSections)
+			{
+				if (const UMovieSceneSubSection* StrongSubSection = SubSection.Get())
+				{
+					if (bool* SectionToPreserveFlag = SectionsToPreserve->Find(StrongSubSection))
+					{
+						*SectionToPreserveFlag = false;
+					}
+				}
+			}
+		}
+	}
+
+	return bReinit;
+}
+#endif
+
+void ADaySequenceActor::InitializeDaySequences()
+{
+	SubSections.Empty();
+	
+	if (DaySequenceCollection)
+	{
+		for (const FDaySequenceCollectionEntry& Entry : DaySequenceCollection->DaySequences)
+		{
+			InitializeDaySequence(Entry);
+		}
+	}
+}
+
+UMovieSceneSubSection* ADaySequenceActor::InitializeDaySequence(const FDaySequenceCollectionEntry& SequenceAsset)
+{
+	if (!RootSequence || !SequenceAsset.Sequence)
+	{
+		return nullptr;
+	}
+
+	UMovieScene* RootMovieScene = RootSequence->GetMovieScene();
+	UDaySequenceTrack* SubTrack = RootMovieScene->AddTrack<UDaySequenceTrack>();
+	UMovieSceneSubSection* SubSection = nullptr;
+
+	if (SubTrack)
+	{
+		SubTrack->ClearFlags(RF_Transactional);
+		SubTrack->SetFlags(RF_Transient);
+
+		// Add the subsequence section with an arbitrary duration. It will be
+		// normalized in UpdateSubSectionTimeScale.
+		const int32 RootDuration = RootMovieScene->GetPlaybackRange().GetUpperBoundValue().Value;
+		SubSection = SubTrack->AddSequence(SequenceAsset.Sequence, 0, RootDuration);
+
+		if (SubSection)
+		{
+			UpdateSubSectionTimeScale(SubSection);
+			SubSection->Parameters.HierarchicalBias = SequenceAsset.BiasOffset + Bias;
+			SubSection->Parameters.Flags            = EMovieSceneSubSectionFlags::OverrideRestoreState;
+
+			const TFunction<void(void)> SetSubTrackMuteStateConditional = [this, SubSection, Conditions = SequenceAsset.Conditions.Conditions]()
+			{
+				if (!IsValid(this) || !IsValid(SubSection))
+				{
+					return;
+				}
+
+				constexpr bool bInitialMuteState = false;
+				if (const bool bActive = !EvaluateSequenceConditions(bInitialMuteState, Conditions); SubSection->IsActive() != bActive)
+				{
+					SubSection->MarkAsChanged();
+					SubSection->SetIsActive(bActive);
+				}
+			};
+
+			const TFunction<void(void)> SetSubTrackMuteStateUnconditional = [this, SubSection]()
+			{
+				if (!IsValid(this) || !IsValid(SubSection))
+				{
+					return;
+				}
+
+				if (!SubSection->IsActive())
+				{
+					SubSection->MarkAsChanged();
+					SubSection->SetIsActive(true);
+				}
+			};
+			
+			const TFunction<void(void)>& SetSubTrackMuteState = SequenceAsset.Conditions.Conditions.Num() == 0 ? SetSubTrackMuteStateUnconditional : SetSubTrackMuteStateConditional;
+			
+			// Initialize mute state and set up the condition callbacks to dynamically update mute state.
+			SetSubTrackMuteState();
+			OnInvalidateMuteStates.AddWeakLambda(SubSection, SetSubTrackMuteState);
+			BindToConditionCallbacks(SubSection, SequenceAsset.Conditions.Conditions, [this]() { InvalidateMuteStates(); });
+
+			SubSections.Add(SubSection);
+		}
+		else
+		{
+			UE_LOG(LogDaySequence, Warning, TEXT("Failed to create UMovieSceneSubSection in ADaySequenceActor::InitializeDaySequence"));
+		}
+	}
+	else
+	{
+		UE_LOG(LogDaySequence, Warning, TEXT("Failed to create UDaySequenceTrack in ADaySequenceActor::InitializeDaySequence"));
+	}
+
+#if ENABLE_DRAW_DEBUG
+	if (SubSection)
+	{
+        TSharedPtr<TMap<FString, FString>> DebugData = MakeShared<TMap<FString, FString>>();
+        SubSectionDebugEntries.Emplace(MakeShared<UE::DaySequence::FDaySequenceDebugEntry>(
+        [this](){ return true; },
+        [this, DebugData, SubSection]()
+        {
+        	if (IsValid(SubSection))
+        	{
+        		(*DebugData).FindOrAdd("Actor Name") = GetFName().ToString();
+        		(*DebugData).FindOrAdd("Sequence Name") = SubSection->GetSequence() ? SubSection->GetSequence()->GetFName().ToString() : "None";
+        		(*DebugData).FindOrAdd("Mute State") = SubSection->IsActive() ? "Active" : "Muted";
+        		(*DebugData).FindOrAdd("Hierarchical Bias") = FString::Printf(TEXT("%d"), SubSection->Parameters.HierarchicalBias);
+        	}
+        	return DebugData;
+        }));
+	}
+#endif
+
+	return SubSection;
+}
+
+void ADaySequenceActor::UpdateSubSectionTimeScale(UMovieSceneSubSection* InSubSection) const
+{
+	if (!InSubSection || !RootSequence)
+	{
+		return;
+	}
+	
+	const UMovieSceneSequence* Sequence = InSubSection->GetSequence();
+	if (!Sequence)
+	{
+		return;
+	}
+	
+	// Compute outer duration from subsequence asset.
+	const UMovieScene* MovieScene = Sequence->GetMovieScene();
+	const FFrameRate TickResolution = MovieScene->GetTickResolution();
+	const FQualifiedFrameTime InnerDuration = FQualifiedFrameTime(
+		UE::MovieScene::DiscreteSize(MovieScene->GetPlaybackRange()),
+		TickResolution);
+
+	const UDaySequenceTrack* SubTrack = InSubSection->GetTypedOuter<UDaySequenceTrack>();
+	const FFrameRate OuterFrameRate = SubTrack->GetTypedOuter<UMovieScene>()->GetTickResolution();
+	const int32      OuterDuration  = InnerDuration.ConvertTo(OuterFrameRate).FrameNumber.Value;
+
+	// Set the subsequence section to span the full day cycle range and
+	// normalize playback by setting TimeScale on the section.
+	const UMovieScene* RootMovieScene = RootSequence->GetMovieScene();
+	const int32 RootDuration = RootMovieScene->GetPlaybackRange().GetUpperBoundValue().Value;
+	InSubSection->MarkAsChanged();
+	InSubSection->Parameters.TimeScale = (float)OuterDuration / (float)RootDuration;
+	InSubSection->SetRange(RootMovieScene->GetPlaybackRange());
+}
+
+void ADaySequenceActor::OnSequencePlayerUpdate(const UMovieSceneSequencePlayer& Player, FFrameTime CurrentTime, FFrameTime PreviousTime)
+{
+	CSV_SCOPED_TIMING_STAT_EXCLUSIVE(DaySequenceActor_OnSequencePlayerUpdate);
+	
+	auto FrameTimeToDayHours = [this](const FFrameTime& Time)
+	{
+		const FFrameRate FrameRate = RootSequence->GetMovieScene()->GetDisplayRate();
+		const double CurrentTimeSeconds = FQualifiedFrameTime(Time, FrameRate).AsSeconds();
+		const float DayCycleSeconds = DaySequenceHelpers::TimecodeToSeconds(TimePerCycle, FrameRate);
+		const float SequenceRatio = CurrentTimeSeconds / DayCycleSeconds;
+		return DaySequenceHelpers::TimecodeToHours(DayLength, FrameRate) * SequenceRatio;
+	};
+	const float CurrentHours = FrameTimeToDayHours(CurrentTime);
+	const float PreviousHours = FrameTimeToDayHours(PreviousTime);
+	SequencePlayerUpdated(CurrentHours, PreviousHours);
+}
+
+void ADaySequenceActor::SequencePlayerUpdated(float CurrentTime, float PreviousTime)
+{
+}
+
+#if WITH_EDITORONLY_DATA
+bool ADaySequenceActor::GetOverrideInitialTimeOfDay() const
+{
+	return bOverrideInitialTimeOfDay;
+}
+
+void ADaySequenceActor::SetOverrideInitialTimeOfDay(bool bNewOverrideInitialTimeOfDay)
+{
+	SetOverrideInitialTimeOfDay(bNewOverrideInitialTimeOfDay, GetTimeOfDayPreview());
+}
+
+void ADaySequenceActor::SetOverrideInitialTimeOfDay(bool bNewOverrideInitialTimeOfDay, float OverrideInitialTimeOfDay)
+{
+	bool bNeedsBroadcast = false;
+
+	if (bOverrideInitialTimeOfDay != bNewOverrideInitialTimeOfDay)
+	{
+		bOverrideInitialTimeOfDay = bNewOverrideInitialTimeOfDay;
+		bNeedsBroadcast = true;
+	}
+
+	if (!FMath::IsNearlyEqual(OverrideInitialTimeOfDay, GetTimeOfDayPreview()))
+	{
+		SetTimeOfDayPreview(OverrideInitialTimeOfDay);
+		bNeedsBroadcast = true;
+	}
+
+	if (bNeedsBroadcast)
+	{
+		OnOverrideInitialTimeOfDayChanged.Broadcast(bOverrideInitialTimeOfDay, OverrideInitialTimeOfDay);
+	}
+}
+
+bool ADaySequenceActor::GetOverrideRunDayCycle() const
+{
+	return bOverrideRunDayCycle;
+}
+
+void ADaySequenceActor::SetOverrideRunDayCycle(bool bNewOverrideRunDayCycle)
+{
+	if (bOverrideRunDayCycle != bNewOverrideRunDayCycle)
+	{
+		bOverrideRunDayCycle = bNewOverrideRunDayCycle;
+		OnOverrideRunDayCycleChanged.Broadcast(bOverrideRunDayCycle);
+	}
+}
+#endif
+
+bool ADaySequenceActor::HasValidRootSequence() const
+{
+	return IsValid(RootSequence);
+}
+
+bool ADaySequenceActor::RootSequenceHasValidSections() const
+{
+	if (IsValid(RootSequence))
+	{
+		if (const UMovieScene* RootMovieScene = RootSequence->GetMovieScene())
+		{
+			for (const UMovieSceneSection* Section : RootMovieScene->GetAllSections())
+			{
+				if (IsValid(Section))
+				{
+					return true;
+				}
+			}
+		}
+	}
+	
+	return false;
+}
+
+void ADaySequenceActor::InvalidateMuteStates() const
+{
+	OnInvalidateMuteStates.Broadcast();
+}
+
+void ADaySequenceActor::WarpEvaluationRange(FMovieSceneEvaluationRange& InOutRange) const
+{
+	if (!bUseInterpCurve || !DayInterpCurve || !RootSequence || bForceDisableDayInterpCurve)
+	{
+		return;
+	}
+
+	const FFrameRate TickRate   = InOutRange.GetFrameRate();
+	const FFrameRate FrameRate  = RootSequence->GetMovieScene()->GetDisplayRate();
+	const float DayCycleSeconds = DaySequenceHelpers::TimecodeToSeconds(TimePerCycle, FrameRate);
+	const float DayLengthHours  = DaySequenceHelpers::TimecodeToHours(DayLength, FrameRate);
+
+	TRange<FFrameTime> Range = InOutRange.GetRange();
+
+	// Warp the lower bound if possible
+	if (Range.GetLowerBound().IsClosed())
+	{
+		float TimeSeconds = static_cast<float>(Range.GetLowerBoundValue() / TickRate);
+		float TimeInHours = DayLengthHours * TimeSeconds / DayCycleSeconds;
+
+		// Remap the time
+		TimeInHours = DayInterpCurve->FloatCurve.Eval(TimeInHours, TimeInHours);
+
+		// Set the lower bound value while retaining the inclusivity
+		Range.SetLowerBoundValue((TimeInHours * DayCycleSeconds / DayLengthHours) * TickRate);
+	}
+
+	// Warp the upper bound if possible
+	if (Range.GetUpperBound().IsClosed())
+	{
+		float TimeSeconds = static_cast<float>(Range.GetUpperBoundValue() / TickRate);
+		float TimeInHours = DayLengthHours * TimeSeconds / DayCycleSeconds;
+
+		// Remap the time
+		TimeInHours = DayInterpCurve->FloatCurve.Eval(TimeInHours, TimeInHours);
+
+		// Set the upper bound value while retaining the inclusivity
+		Range.SetUpperBoundValue((TimeInHours * DayCycleSeconds / DayLengthHours) * TickRate);
+	}
+
+	InOutRange.ResetRange(Range);
+}
+
+float ADaySequenceActor::GetDayLength() const
+{
+	const TOptional<FFrameRate> FrameRate = RootSequence ? RootSequence->GetMovieScene()->GetDisplayRate() : TOptional<FFrameRate>();
+	return DaySequenceHelpers::TimecodeToHours(DayLength, FrameRate);
+}
+
+void ADaySequenceActor::SetDayLength(float InHours)
+{
+	const TOptional<FFrameRate> FrameRate = RootSequence ? RootSequence->GetMovieScene()->GetDisplayRate() : TOptional<FFrameRate>();
+
+	// Set min day length to 1 second.
+	const FTimecode OneSecond(0, 0, 1, 0, false);
+	InHours = FMath::Max(InHours, DaySequenceHelpers::TimecodeToHours(OneSecond, FrameRate));
+	DayLength = DaySequenceHelpers::HoursToTimecode(InHours, FrameRate);
+}
+
+float ADaySequenceActor::GetTimePerCycle() const
+{
+	const TOptional<FFrameRate> FrameRate = RootSequence ? RootSequence->GetMovieScene()->GetDisplayRate() : TOptional<FFrameRate>();
+	return DaySequenceHelpers::TimecodeToHours(TimePerCycle, FrameRate);
+}
+
+void ADaySequenceActor::SetTimePerCycle(float InHours)
+{
+	const TOptional<FFrameRate> FrameRate = RootSequence ? RootSequence->GetMovieScene()->GetDisplayRate() : TOptional<FFrameRate>();
+
+	// Set min cycle length to 1 second.
+	const FTimecode OneSecond(0, 0, 1, 0, false);
+	InHours = FMath::Max(InHours, DaySequenceHelpers::TimecodeToHours(OneSecond, FrameRate));
+	TimePerCycle = DaySequenceHelpers::HoursToTimecode(InHours, FrameRate);
+}
+
+void ADaySequenceActor::Multicast_SetTimePerCycle_Implementation(float InHours)
+{
+	const TOptional<FFrameRate> FrameRate = RootSequence ? RootSequence->GetMovieScene()->GetDisplayRate() : TOptional<FFrameRate>();
+
+	const bool bIsGameWorld = GetWorld()->IsGameWorld();
+	const float CurrentTimeOfDay = bIsGameWorld ? GetTimeOfDay() : 0.f;
+
+	// Set min cycle length to 1 second.
+	const FTimecode OneSecond(0, 0, 1, 0, false);
+	InHours = FMath::Max(InHours, DaySequenceHelpers::TimecodeToHours(OneSecond, FrameRate));
+	const FTimecode NewTimePerCycle = DaySequenceHelpers::HoursToTimecode(InHours, FrameRate); 
+	if (NewTimePerCycle != TimePerCycle)
+	{
+		// Validate the new TimePerCycle to avoid overflowing the playback range.
+		UMovieScene* RootMovieScene = RootSequence->GetMovieScene();
+		const FFrameRate DisplayRate = RootMovieScene->GetDisplayRate();
+		const FFrameRate TickResolution = RootMovieScene->GetTickResolution();
+		const auto IsTimePerCycleOverflow = [&DisplayRate, &TickResolution](const FTimecode& InTimePerCycle)
+		{
+			const float TimePerCycleSeconds = DaySequenceHelpers::TimecodeToSeconds(InTimePerCycle, TOptional<FFrameRate>());
+			const int32 CycleSeconds = static_cast<int32>(TimePerCycleSeconds) + 1;
+			const int32 DisplayFactor = DisplayRate.Numerator / DisplayRate.Denominator;
+			const int32 TickFactor = TickResolution.Numerator / TickResolution.Denominator;
+			const int32 NewCycleSeconds = CycleSeconds * DisplayFactor * TickFactor;
+			const int32 CheckCycleSeconds = NewCycleSeconds / (DisplayFactor * TickFactor);
+			return (CheckCycleSeconds != CycleSeconds);
+		};
+		if (IsTimePerCycleOverflow(NewTimePerCycle))
+		{
+			UE_LOG(LogDaySequence, Warning, TEXT("Skipping SetTimePerCycle( %f hours ) to avoid overflowing playback range"), InHours);
+			return;
+		}
+		
+		TimePerCycle = NewTimePerCycle;
+		
+		if (bIsGameWorld)
+		{
+			// Update playback range for the root sequence.
+			const FQualifiedFrameTime FrameTimePerCycle(TimePerCycle, FFrameRate());
+			const float DaySeconds = FrameTimePerCycle.AsSeconds();
+			const int32 RootDuration = RootMovieScene->GetTickResolution().AsFrameNumber(DaySeconds).Value;
+			RootMovieScene->MarkAsChanged();
+			RootMovieScene->SetPlaybackRange(0, RootDuration);
+
+			// Iterate over subsequences and update their time scales.
+			for (UMovieSceneSection* Section : RootMovieScene->GetAllSections())
+			{
+				if (UMovieSceneSubSection* SubSection = Cast<UMovieSceneSubSection>(Section))
+				{
+					UpdateSubSectionTimeScale(SubSection);
+				}
+			}
+
+			// Update the sequence player frame range from the root sequence play range
+			const TRange<FFrameNumber> PlaybackRange = RootMovieScene->GetPlaybackRange();
+			const FFrameNumber SrcStartFrame = UE::MovieScene::DiscreteInclusiveLower(PlaybackRange);
+			const FFrameNumber SrcEndFrame   = UE::MovieScene::DiscreteExclusiveUpper(PlaybackRange);
+
+			const FFrameTime EndingTime = ConvertFrameTime(SrcEndFrame, TickResolution, DisplayRate);
+
+			const FFrameNumber StartingFrame = ConvertFrameTime(SrcStartFrame, TickResolution, DisplayRate).FloorToFrame();
+			const FFrameNumber EndingFrame   = EndingTime.FloorToFrame();
+
+			if (UDaySequencePlayer* Player = GetSequencePlayer())
+			{
+				Player->SetFrameRange(StartingFrame.Value, (EndingFrame - StartingFrame).Value, EndingTime.GetSubFrame());
+			}
+
+			// Finally update the play position to match the current time of day
+			SetTimeOfDay(CurrentTimeOfDay);
+		}
+	}
+}
+
+float ADaySequenceActor::GetInitialTimeOfDay() const
+{
+	const TOptional<FFrameRate> FrameRate = RootSequence ? RootSequence->GetMovieScene()->GetDisplayRate() : TOptional<FFrameRate>();
+	return DaySequenceHelpers::TimecodeToHours(InitialTimeOfDay, FrameRate);
+}
+
+void ADaySequenceActor::SetInitialTimeOfDay(float InHours)
+{
+	InHours = FMath::Clamp(InHours, 0.0f, GetDayLength()); 
+	const TOptional<FFrameRate> FrameRate = RootSequence ? RootSequence->GetMovieScene()->GetDisplayRate() : TOptional<FFrameRate>();
+	InitialTimeOfDay = DaySequenceHelpers::HoursToTimecode(InHours, FrameRate);
+}
+
+float ADaySequenceActor::GetTimeOfDay() const
+{
+	float Result;
+	UWorld* World = GetWorld();
+	const UDaySequencePlayer* Player = GetSequencePlayer();
+	if (HasValidRootSequence() && Player && World && World->IsGameWorld())
+	{
+		const FFrameRate FrameRate = RootSequence->GetMovieScene()->GetDisplayRate();
+		const FQualifiedFrameTime CurrentFrameTime = Player->GetCurrentTime();
+		const double CurrentTimeSeconds = CurrentFrameTime.AsSeconds();
+		const float DayCycleSeconds = DaySequenceHelpers::TimecodeToSeconds(TimePerCycle, FrameRate);
+		const float SequenceRatio = CurrentTimeSeconds / DayCycleSeconds;
+		Result = DaySequenceHelpers::TimecodeToHours(DayLength, FrameRate) * SequenceRatio;
+	}
+	else
+	{
+		Result = GetInitialTimeOfDay();
+	}
+	return Result;
+}
+
+bool ADaySequenceActor::SetTimeOfDay(float InHours)
+{
+	// Only set time of day if we have a valid playing day sequence.
+	UWorld* World = GetWorld();
+	UDaySequencePlayer* Player = GetSequencePlayer();
+	if (HasValidRootSequence() && Player && World && World->IsGameWorld())
+	{
+		// Convert the day time to sequence time
+		const FFrameRate FrameRate = RootSequence->GetMovieScene()->GetDisplayRate();
+		const float DayLengthHours = GetDayLength();
+		const float DayLengthRatio = FMath::Frac(InHours / DayLengthHours);
+		const float DayCycleSeconds = DaySequenceHelpers::TimecodeToSeconds(TimePerCycle, FrameRate) * DayLengthRatio;
+
+		// Update the playback position of the sequence.
+		FMovieSceneSequencePlaybackParams PlaybackParams;
+		PlaybackParams.Frame = FrameRate.AsFrameTime(DayCycleSeconds); 
+		PlaybackParams.PositionType = EMovieScenePositionType::Frame;
+		PlaybackParams.UpdateMethod = EUpdatePositionMethod::Play;
+		Player->SetPlaybackPosition(PlaybackParams);
+		return true;
+	}
+	return false;
+}
+
+bool ADaySequenceActor::HasStaticTimeOfDay() const
+{
+#if WITH_EDITORONLY_DATA
+	return WeakTimeControllerOverride.Pin() != nullptr || StaticTime.IsSet();
+#else
+	return WeakTimeControllerOverride.Pin() != nullptr;
+#endif
+}
+
+float ADaySequenceActor::GetStaticTimeOfDay() const
+{
+	using namespace UE::DaySequence;
+
+	if (TSharedPtr<FStaticTimeControllerOverride> TimeControllerOverride = WeakTimeControllerOverride.Pin())
+	{
+		return TimeControllerOverride->TimeInHours;
+	}
+
+#if WITH_EDITORONLY_DATA
+	if (StaticTime.IsSet())
+	{
+		return StaticTime.GetValue();
+	}
+#endif
+	
+	return std::numeric_limits<float>::lowest();
+}
+
+void ADaySequenceActor::SetStaticTimeOfDay(float InHours)
+{
+	using namespace UE::DaySequence;
+
+#if WITH_EDITORONLY_DATA
+	StaticTime = InHours;
+#endif
+	
+	// Convert the day time to sequence time
+	const FFrameRate FrameRate = RootSequence->GetMovieScene()->GetDisplayRate();
+	const float DayLengthHours = GetDayLength();
+	const float DayLengthRatio = FMath::Frac(InHours / DayLengthHours);
+	const float DayCycleSeconds = DaySequenceHelpers::TimecodeToSeconds(TimePerCycle, FrameRate) * DayLengthRatio;
+	const FFrameTime StaticFrameTime = FrameRate.AsFrameTime(DayCycleSeconds);
+
+	if (TSharedPtr<FStaticTimeControllerOverride> ExistingTimeControllerOverride = WeakTimeControllerOverride.Pin())
+	{
+		ExistingTimeControllerOverride->StaticTime = StaticFrameTime;
+		ExistingTimeControllerOverride->TimeInHours = InHours;
+	}
+	else
+	{
+		if (UDaySequencePlayer* Player = GetSequencePlayer())
+		{
+			TSharedRef<FStaticTimeControllerOverride> TimeControllerOverride =
+				MakeShared<FStaticTimeControllerOverride>(StaticFrameTime, Player->GetTimeController());
+
+			TimeControllerOverride->TimeInHours = InHours;
+#if WITH_EDITORONLY_DATA
+			TimeControllerOverride->PreviousTimeOfDayPreview = TimeOfDayPreview;
+#endif
+
+			Player->SetTimeControllerDirectly(TimeControllerOverride);
+			Player->SetIgnorePlaybackReplication(true);
+
+			WeakTimeControllerOverride = TimeControllerOverride;
+		}
+	}
+
+#if WITH_EDITOR
+	ConditionalSetTimeOfDayPreview(InHours);
+#endif
+}
+
+void ADaySequenceActor::RemoveStaticTimeOfDay(bool bResumeFromStaticTime)
+{
+	using namespace UE::DaySequence;
+
+#if WITH_EDITORONLY_DATA
+	StaticTime.Reset();
+#endif
+	
+	TSharedPtr<FStaticTimeControllerOverride> TimeControllerOverride = WeakTimeControllerOverride.Pin();
+	UDaySequencePlayer* Player = GetSequencePlayer();
+	if (TimeControllerOverride && Player)
+	{
+		if (bResumeFromStaticTime && TimeControllerOverride->PreviousTimeController)
+		{
+			// Reset the old time controller to the static time if we want to resume from here
+			TimeControllerOverride->PreviousTimeController->Reset(Player->GetCurrentTime());
+		}
+
+		// Restore the old time controller and re-enable network replication
+		Player->SetTimeControllerDirectly(TimeControllerOverride->PreviousTimeController);
+		Player->SetIgnorePlaybackReplication(false);
+
+#if WITH_EDITOR
+		// Update the time of day preview if we're in an editor world
+		UWorld* World = GetWorld();
+		if (World && World->WorldType == EWorldType::Editor)
+		{
+			if (!bResumeFromStaticTime)
+			{
+				TimeOfDayPreview = TimeControllerOverride->PreviousTimeOfDayPreview;
+				OnTimeOfDayPreviewChangedEvent.Broadcast(GetTimeOfDayPreview());
+			}
+		}
+#endif
+
+		WeakTimeControllerOverride = nullptr;
+	}
+}
+
+void ADaySequenceActor::Play()
+{
+	UDaySequencePlayer* Player = GetSequencePlayer();
+	UWorld* World = GetWorld();
+	if (HasValidRootSequence() && Player && World && World->IsGameWorld())
+	{
+		// Always ensure play is looping.
+		Player->PlayLooping();
+	}
+}
+
+void ADaySequenceActor::Pause()
+{
+	UDaySequencePlayer* Player = GetSequencePlayer();
+	UWorld* World = GetWorld();
+	if (HasValidRootSequence() && Player && World && World->IsGameWorld())
+	{
+		Player->Pause();
+	}
+}
+
+bool ADaySequenceActor::IsPlaying() const
+{
+	if (const UDaySequencePlayer* Player = GetSequencePlayer())
+	{
+		return Player->IsPlaying();
+	}
+
+	return false;
+}
+
+bool ADaySequenceActor::IsPaused() const
+{
+	if (const UDaySequencePlayer* Player = GetSequencePlayer())
+	{
+		return Player->IsPaused();
+	}
+
+	return false;
+}
+
+UDaySequence* ADaySequenceActor::GetRootSequence() const
+{
+	return RootSequence;
+}
+
+void ADaySequenceActor::UpdateRootSequence()
+{
+#if WITH_EDITOR
+	FCoreUObjectDelegates::OnObjectModified.RemoveAll(this);
+	
+	// Set up a callback for when UObjects are modified so we can catch changes to our 
+	FCoreUObjectDelegates::OnObjectModified.AddWeakLambda(this, [this](UObject* InObject)
+	{
+		if (IsValid(InObject) && DaySequenceCollection == InObject)
+		{
+			// We update on next tick because calling UpdateScalabilitySequence here
+			// is too early (our collection doesn't have the changes that triggered this invocation).
+			bUpdateRootSequenceOnTick = true;
+
+			// Must be empty for a full reinit to take place (on next tick)
+			SubSections.Empty();
+		}
+	});
+#endif
+	
+	// For now, just rebuild the root sequence and GC the old one.
+	InitializeRootSequence();
+}
+
+#if WITH_EDITOR
+void ADaySequenceActor::UpdateRootSequenceOnTick()
+{
+	bUpdateRootSequenceOnTick = true;
+}
+#endif
+
+
+/**
+ * Compute a PlaybackSettings object for the given DaySequence with a fixed 1.0x PlayRate.
+ *
+ * This method is used to workaround the issue where a non-1.0x PlayRate causes issues with
+ * sequence playback replication.
+ */
+FMovieSceneSequencePlaybackSettings ADaySequenceActor::GetPlaybackSettings(const UDaySequence* Sequence) const
+{
+	FMovieSceneSequencePlaybackSettings Settings;
+	Settings.bAutoPlay = true;
+	Settings.LoopCount.Value = -1; // Loop indefinitely
+	Settings.bDisableCameraCuts = true;
+	Settings.PlayRate = 1.0f;
+	Settings.StartTime = 0.0f;
+
+	// Always explicitly inherit the tick interval from this actor by initializing with 'this'
+	Settings.TickInterval = FMovieSceneSequenceTickInterval(this);
+	// Set explicit frame budget based on the cvar
+	Settings.TickInterval.EvaluationBudgetMicroseconds = UE::DaySequence::GFrameBudgetMicroseconds;
+	// No longer inherit the tick interval since that would wipe out our budget
+	Settings.bInheritTickIntervalFromOwner = false;
+
+	if (Sequence && Sequence->GetMovieScene())
+	{
+		// Convert the initial day time to sequence time
+		const UMovieScene* MovieScene = Sequence->GetMovieScene();
+		const FFrameRate DisplayRate = MovieScene->GetDisplayRate();
+
+#if WITH_EDITORONLY_DATA
+		const float InitialHours = bOverrideInitialTimeOfDay ? GetTimeOfDayPreview() : GetInitialTimeOfDay();
+#else
+		const float InitialHours = GetInitialTimeOfDay();
+#endif
+		
+		const float DayLengthHours = GetDayLength();
+		const float StartRatio = FMath::Frac(InitialHours / DayLengthHours);
+		const float DayCycleSeconds = DaySequenceHelpers::TimecodeToSeconds(TimePerCycle, DisplayRate) * StartRatio;
+		Settings.StartTime = DayCycleSeconds;
+	}
+	return Settings;
+}
+
+#if ENABLE_DRAW_DEBUG
+bool ADaySequenceActor::IsDebugCategoryRegistered(const FName& Category) const
+{
+	return DebugEntries.Find(Category) != nullptr;
+}
+
+void ADaySequenceActor::RegisterDebugCategory(const FName& Category, UE::DaySequence::FDebugCategoryDrawFunction DrawFunction)
+{
+	using namespace UE::DaySequence;
+	
+	if (TPair<FDebugEntryArray, FDebugCategoryDrawFunction>* CategoryHandle = DebugEntries.Find(Category))
+	{
+		ensureMsgf(false, TEXT("Category \"%s\" is already registered with this DaySequenceActor."), *Category.ToString());
+		return;
+	}
+
+	DebugEntries.Add(Category, TPair<FDebugEntryArray, FDebugCategoryDrawFunction>(FDebugEntryArray(), DrawFunction));
+}
+
+void ADaySequenceActor::RegisterDebugEntry(const TWeakPtr<UE::DaySequence::FDaySequenceDebugEntry>& DebugEntry, const FName& Category)
+{
+	using namespace UE::DaySequence;
+	
+	if (TPair<FDebugEntryArray, FDebugCategoryDrawFunction>* CategoryHandle = DebugEntries.Find(Category))
+	{
+		FDebugEntryArray& Entries = CategoryHandle->Key;
+		Entries.AddUnique(DebugEntry);
+	}
+	else
+	{
+		ensureMsgf(false, TEXT("Category \"%s\" is not registered with this DaySequenceActor."), *Category.ToString());
+	}
+}
+
+void ADaySequenceActor::UnregisterDebugEntry(const TWeakPtr<UE::DaySequence::FDaySequenceDebugEntry>& DebugEntry, const FName& Category)
+{
+	using namespace UE::DaySequence;
+	
+	if (TPair<FDebugEntryArray, FDebugCategoryDrawFunction>* CategoryValue = DebugEntries.Find(Category))
+	{
+		FDebugEntryArray& Entries = CategoryValue->Key;
+		Entries.Remove(DebugEntry);
+	}
+	else
+	{
+		ensureMsgf(false, TEXT("Category \"%s\" is not registered with this DaySequenceActor."), *Category.ToString());
+	}
+}
+
+void ADaySequenceActor::OnShowDebugInfoDrawFunction(UCanvas* Canvas, TArray<TSharedPtr<TMap<FString, FString>>>& Entries, const FString& Category)
+{
+	using namespace UE::DaySequence;
+
+	// Early out if this table will be empty (occurs if the Entry->ShowCondition() condition above never evaluates to true)
+	if (Entries.Num() == 0)
+	{
+		return;
+	}
+	
+	FDisplayDebugManager& DisplayDebugManager = Canvas->DisplayDebugManager;
+	const UFont* Font = GEngine->GetSmallFont();
+	DisplayDebugManager.SetFont(Font);
+	
+	// Used for padding table entries. Char count isn't sufficient as these fonts aren't monospace so we must
+	// compute the number of spaces when we need to pad as the width in pixels of the area we need to pad divided
+	// by the width of a single space character and then pad appropriately
+	float SpaceCharacterWidth = 0.f;
+	float SpaceCharacterHeight = 0.f;
+	Font->GetCharSize(' ', SpaceCharacterWidth, SpaceCharacterHeight);
+
+	const FString FieldSeparator = "    |    ";
+	const int32 FieldSeparatorSize = Font->GetStringSize(&FieldSeparator[0]);
+	
+	// Determine column widths based on the largest value to be printed in each column (including column labels)
+	// Also keep track of the running total row length which is simply the sum of the columns lengths
+	TMap<FString, int32> LongestPropertyValues;
+	int32 TotalExpectedRowLength = 0;
+	for (const auto& Entry : Entries)
+	{
+		for (const auto& Cell : *Entry)
+		{
+			const int32 NewPropertySize = Font->GetStringSize(&Cell.Value[0]);
+
+			if (int32* ExistingPropertySize = LongestPropertyValues.Find(Cell.Key))
+			{
+				// Update existing property's maximum known size
+				
+				const int32 IncreaseInPropertySize = FMath::Max(NewPropertySize - *ExistingPropertySize, 0);
+				*ExistingPropertySize += IncreaseInPropertySize;
+				TotalExpectedRowLength += IncreaseInPropertySize;
+
+				// This is the more obvious way to do this but makes tracking the running total more annoying
+				// *ExistingPropertySize = FMath::Max(*ExistingPropertySize, NewPropertySize);
+			}
+			else
+			{
+				// Add new length entry that is max(StringSize, PropertyNameStringSize) because the property
+				// names are printed in their own row and should factor into the column widths
+
+				const int32 IncreaseInRowLength = FMath::Max(NewPropertySize, Font->GetStringSize(&Cell.Key[0]));
+				LongestPropertyValues.Add(Cell.Key, IncreaseInRowLength);
+				TotalExpectedRowLength += IncreaseInRowLength;
+			}
+			
+			int32& Length = LongestPropertyValues.FindOrAdd(Cell.Key, NewPropertySize);
+			Length = FMath::Max(Length, NewPropertySize);
+		}
+	}
+	// If we have N columns then there are N - 1 separators between them.
+	TotalExpectedRowLength += (LongestPropertyValues.Num() - 1) * FieldSeparatorSize;
+
+	/* BEGIN DRAWING HERE */
+
+	auto PadToCenterString = [Font, SpaceCharacterWidth](const FString& StringToCenter, const int32 DesiredStringLength)
+	{
+		const int32 StringToCenterSize = Font->GetStringSize(&StringToCenter[0]);
+		const int32 PadSpaceCount = FMath::CeilToInt((DesiredStringLength - StringToCenterSize) / SpaceCharacterWidth);
+		const int32 PadLeftCount = PadSpaceCount / 2;
+		const int32 PadRightCount = PadSpaceCount / 2 + PadSpaceCount % 2;
+
+		// LeftPad/RightPad will attempt to pad with space such that the returned string's length is ChCount (the single parameter)
+		// It is NOT padding by ChCount characters and it does not modify the internal string, it returns a copy.
+		// So first we pad to the left by the current string size + the spaces we want on the left then we pad the returned
+		// string by the original string size + the desired left spaces + the desired right spaces.
+		return StringToCenter.LeftPad(StringToCenter.Len() + PadLeftCount).RightPad(StringToCenter.Len() + PadLeftCount + PadRightCount);
+	};
+
+	// Print some new lines to clearly separate this table from any previous data.
+	// A better way to do this would be to using the DisplayDebugManagers SetYPos() function but it isn't exported.
+	constexpr int32 NumInitialLineSkips = 3;
+	for (int Line = 0; Line < NumInitialLineSkips; ++Line)
+	{
+		DisplayDebugManager.DrawString("");
+	}
+	
+	// Print header text
+	{
+		const FString RowText = FieldSeparator + PadToCenterString("Category: " + Category, TotalExpectedRowLength) + FieldSeparator;
+		DisplayDebugManager.SetDrawColor(FColor::Yellow);
+		DisplayDebugManager.DrawString(RowText);
+	}
+	
+	// Print column labels
+	{
+		FString RowText = FieldSeparator;
+		
+		for (const auto& Property : LongestPropertyValues)
+		{
+			RowText += PadToCenterString(Property.Key, Property.Value) + FieldSeparator;
+		}
+	
+		DisplayDebugManager.DrawString(RowText);
+	}
+
+	// Print column values
+	DisplayDebugManager.SetDrawColor(FColor::White);
+	for (const auto& Entry : Entries)
+	{
+		TMap<FString, FString>& EntryData = *Entry;
+		FString RowText = FieldSeparator;
+
+		for (const auto& Property : LongestPropertyValues)
+		{
+			RowText += PadToCenterString(EntryData.Find(Property.Key) ? EntryData[Property.Key] : "None", Property.Value) + FieldSeparator;
+		}
+
+		DisplayDebugManager.DrawString(RowText);
+	}
+}
+
+void ADaySequenceActor::OnShowDebugInfo(AHUD* HUD, UCanvas* Canvas, const FDebugDisplayInfo& DisplayInfo, float& YL, float& YPos)
+{
+	using namespace UE::DaySequence;
+	
+	if (!Canvas || !GEngine || HasAuthority())
+	{
+		return;
+	}
+
+	DebugEntries = DebugEntries.FilterByPredicate([](const TPair<FName, TPair<FDebugEntryArray, FDebugCategoryDrawFunction>>& Category) {
+		return !Category.Value.Key.IsEmpty();
+	});
+
+	auto RemoveStaleEntriesAndGetPinnedArray = [](FDebugEntryArray& WeakArray, TArray<TSharedPtr<TMap<FString, FString>>>& OutSharedArray)
+	{
+		// Shouldn't reduce existing capacity, so in theory the total number
+		// of allocations here will be <= the size of the largest array in DebugEntries
+		OutSharedArray.Reset();
+		
+		WeakArray.RemoveAll([](const TWeakPtr<FDaySequenceDebugEntry>& Entry) 
+			{ return !Entry.IsValid(); });
+
+		for (TWeakPtr<FDaySequenceDebugEntry> Entry : WeakArray)
+		{
+			OutSharedArray.Push(Entry.Pin()->GetData());
+		}
+	};
+	
+	TArray<TSharedPtr<TMap<FString, FString>>> EntriesToDraw;
+	for (TPair<FName, TPair<FDebugEntryArray, FDebugCategoryDrawFunction>>& Category : DebugEntries)
+	{
+		// Print this category if it is individually enabled via "showdebug Category.Key" OR if
+		// the general DaySequenceCategory is enabled via "showdebug DaySequence"
+		if (HUD->ShouldDisplayDebug(ShowDebug_GeneralCategory) || HUD->ShouldDisplayDebug(Category.Key))
+		{
+			FDebugEntryArray& Entries = Category.Value.Key;
+			FDebugCategoryDrawFunction& DrawFunction = Category.Value.Value;
+
+			RemoveStaleEntriesAndGetPinnedArray(Entries, EntriesToDraw);
+
+			DrawFunction(Canvas, EntriesToDraw, Category.Key.ToString());
+		}
+	}
+}
+#endif
+
+TObjectPtr<UDaySequenceConditionTag> ADaySequenceActor::GetOrInstantiateConditionTag(const TSubclassOf<UDaySequenceConditionTag>& ConditionClass)
+{
+	/**
+	 * If ConditionClass is nullptr, ConditionClass is not a child of UDaySequenceConditionTag, or we are PostLoading, catch here.
+	 * Note: *ConditionClass returns a UClass* (does not dereference pointers), so *ConditionClass == nullptr is a null check.
+	 * Early out if we are PostLoading because we can't safely call UDaySequenceConditionTag::Initialize (we will reinitialize sequences soon after this).
+	 */
+	if (*ConditionClass == nullptr || !ConditionClass->IsChildOf(UDaySequenceConditionTag::StaticClass()) || FUObjectThreadContext::Get().IsRoutingPostLoad)
+	{
+		return nullptr;
+	}
+
+	// Instantiate condition if necessary
+	if (!TrackConditionMap.Contains(*ConditionClass) || (TrackConditionMap.Contains(*ConditionClass) && !IsValid(TrackConditionMap[*ConditionClass])))
+	{
+		TrackConditionMap.Remove(*ConditionClass);
+		TrackConditionMap.Add(*ConditionClass, NewObject<UDaySequenceConditionTag>(this, ConditionClass));
+
+		TrackConditionMap[*ConditionClass]->Initialize();
+	}
+
+	return TrackConditionMap[*ConditionClass];
+}
+
+bool ADaySequenceActor::EvaluateSequenceConditions(bool bInitialMuteState, const FDaySequenceConditionSet::FConditionValueMap& InConditions)
+{
+	bool bMuted = bInitialMuteState;
+	
+	for (const TTuple<TSubclassOf<UDaySequenceConditionTag>, bool>& Element : InConditions)
+	{
+		const TSubclassOf<UDaySequenceConditionTag>& ConditionClass = Element.Key;
+		const bool& ExpectedValue = Element.Value;
+		
+		if (const TObjectPtr<UDaySequenceConditionTag> ConditionInstance = GetOrInstantiateConditionTag(ConditionClass))
+		{
+			// If ConditionInstance evaluates to ExpectedValue, bMuted is unchanged. Otherwise, bMuted is set to true.
+			bMuted |= ConditionInstance->Evaluate() != ExpectedValue;
+		}
+	}
+
+	return bMuted;
+}
+
+void ADaySequenceActor::BindToConditionCallbacks(UObject* LifetimeObject, const FDaySequenceConditionSet::FConditionValueMap& InConditions, const TFunction<void(void)>& InFunction)
+{
+	for (const TTuple<TSubclassOf<UDaySequenceConditionTag>, bool>& Element : InConditions)
+	{
+		const TSubclassOf<UDaySequenceConditionTag>& ConditionClass = Element.Key;
+		
+		if (const TObjectPtr<UDaySequenceConditionTag> ConditionInstance = GetOrInstantiateConditionTag(ConditionClass))
+		{
+			// Invoke InFunction when ConditionInstance's OnConditionValueChanged delegate is broadcast if LifetimeObject is still valid.
+			ConditionInstance->GetOnConditionValueChanged().AddWeakLambda(LifetimeObject, InFunction);
+		}
+	}
+}
diff --git a/Engine/Plugins/Experimental/DaySequence/Source/DaySequence/Private/DaySequenceConditionTag.cpp b/Engine/Plugins/Experimental/DaySequence/Source/DaySequence/Private/DaySequenceConditionTag.cpp
new file mode 100644
--- /dev/null
+++ b/Engine/Plugins/Experimental/DaySequence/Source/DaySequence/Private/DaySequenceConditionTag.cpp
@@ -0,0 +1,42 @@
+﻿// Copyright Epic Games, Inc. All Rights Reserved.
+
+#include "DaySequenceConditionTag.h"
+
+void UDaySequenceConditionTag::Initialize()
+{
+	SetupBroadcastBindings();
+}
+
+bool UDaySequenceConditionTag::Evaluate_Implementation() const
+{
+	return true;
+}
+
+FString UDaySequenceConditionTag::GetConditionName_Implementation() const
+{
+	return GetClass()->GetName();
+}
+
+UWorld* UDaySequenceConditionTag::GetWorld() const
+{
+	if (IsTemplate())
+	{
+		return nullptr;
+	}
+
+	return GetOuter()->GetWorld();
+}
+
+void UDaySequenceConditionTag::SetupBroadcastBindings_Implementation() const
+{
+}
+
+void UDaySequenceConditionTag::BroadcastOnConditionValueChanged()
+{
+	const bool bResult = Evaluate();
+	if (!bCachedEvalResult.IsSet() || bCachedEvalResult.GetValue() != bResult)
+	{
+		bCachedEvalResult = bResult;
+		OnConditionValueChanged.Broadcast();
+	}
+}
\ No newline at end of file
diff --git a/Engine/Plugins/Experimental/DaySequence/Source/DaySequence/Private/DaySequenceHelpers.cpp b/Engine/Plugins/Experimental/DaySequence/Source/DaySequence/Private/DaySequenceHelpers.cpp
new file mode 100644
--- /dev/null
+++ b/Engine/Plugins/Experimental/DaySequence/Source/DaySequence/Private/DaySequenceHelpers.cpp
@@ -0,0 +1,55 @@
+// Copyright Epic Games, Inc. All Rights Reserved.
+
+#include "DaySequenceHelpers.h"
+#include "Misc/FrameRate.h"
+#include "Misc/QualifiedFrameTime.h"
+
+static constexpr float SecondsPerHour = 3600.f;
+static constexpr float SecondsPerMinute = 60.f;
+
+float DaySequenceHelpers::TimecodeToHours(const FTimecode& InTimecode, const TOptional<FFrameRate>& InRate)
+{
+	return TimecodeToSeconds(InTimecode, InRate) / SecondsPerHour;
+}
+
+float DaySequenceHelpers::TimecodeToSeconds(const FTimecode& InTimecode, const TOptional<FFrameRate>& InRate)
+{
+	float Result;
+	if (InRate.IsSet())
+	{
+		const FQualifiedFrameTime FrameTime(InTimecode, *InRate);
+		Result = FrameTime.AsSeconds();
+	}
+	else
+	{
+		Result = InTimecode.Hours * SecondsPerHour + InTimecode.Minutes * SecondsPerMinute + InTimecode.Seconds; 
+	}
+	return Result;
+}
+
+FTimecode DaySequenceHelpers::HoursToTimecode(float InHours, const TOptional<FFrameRate>& InRate)
+{
+	return SecondsToTimecode(InHours * SecondsPerHour, InRate);
+}
+
+FTimecode DaySequenceHelpers::SecondsToTimecode(float InSeconds, const TOptional<FFrameRate>& InRate)
+{
+	FTimecode Result;
+	if (InRate.IsSet())
+	{
+		Result = FQualifiedFrameTime(InRate->AsFrameTime(InSeconds), *InRate).ToTimecode();
+	}
+	else
+	{
+		float Seconds = InSeconds;
+		const float Hours = FMath::Floor(InSeconds / SecondsPerHour);
+		Seconds -= Hours * SecondsPerHour;
+		const float Minutes = FMath::Floor(Seconds / SecondsPerMinute);
+		Seconds -= Minutes * SecondsPerMinute;
+		Result.Hours = Hours;
+		Result.Minutes = Minutes;
+		Result.Seconds = FMath::Floor(Seconds);
+	}
+	return Result;
+}
+
diff --git a/Engine/Plugins/Experimental/DaySequence/Source/DaySequence/Private/DaySequenceHelpers.h b/Engine/Plugins/Experimental/DaySequence/Source/DaySequence/Private/DaySequenceHelpers.h
new file mode 100644
--- /dev/null
+++ b/Engine/Plugins/Experimental/DaySequence/Source/DaySequence/Private/DaySequenceHelpers.h
@@ -0,0 +1,57 @@
+// Copyright Epic Games, Inc. All Rights Reserved.
+
+#pragma once
+
+#include "Misc/OptionalFwd.h"
+
+struct FTimecode;
+
+struct FFrameRate;
+
+class DaySequenceHelpers
+{
+public:
+	/**
+	 * Convert Timecode to hours. Optionally provide a frame rate to convert
+	 * frames portion of Timecode. If a frame rate is not provided, the
+	 * frames portion of Timecode will be ignored.
+	 *
+	 * @param InTimecode Timecode to convert
+	 * @param InRate Optional frame rate to convert frames portion of Timecode.
+	 * @return result in hours.
+	 */
+	static float TimecodeToHours(const FTimecode& InTimecode, const TOptional<FFrameRate>& InRate);
+
+	/**
+	 * Convert Timecode to seconds. Optionally provide a frame rate to convert
+	 * frames portion of Timecode. If a frame rate is not provided, the
+	 * frames portion of Timecode will be ignored.
+	 *
+	 * @param InTimecode Timecode to convert
+	 * @param InRate Optional frame rate to convert frames portion of Timecode.
+	 * @return result in seconds.
+	 */
+	static float TimecodeToSeconds(const FTimecode& InTimecode, const TOptional<FFrameRate>& InRate);
+
+	/**
+	 * Convert hours to Timecode. Optionally provide a frame rate to convert
+	 * fractional seconds into Timecode frames. If a frame rate is not provided,
+	 * the fractional seconds will be ignored.
+	 *
+	 * @param InHours hours to convert
+	 * @param InRate Optional frame rate to convert fractional seconds to frames.
+	 * @return result in timecode.
+	 */
+	static FTimecode HoursToTimecode(float InHours, const TOptional<FFrameRate>& InRate);
+
+	/**
+	 * Convert seconds to Timecode. Optionally provide a frame rate to convert
+	 * fractional seconds into Timecode frames. If a frame rate is not provided,
+	 * the fractional seconds will be ignored.
+	 *
+	 * @param InSeconds seconds to convert
+	 * @param InRate Optional frame rate to convert fractional seconds to frames.
+	 * @return result in timecode.
+	 */
+	static FTimecode SecondsToTimecode(float InSeconds, const TOptional<FFrameRate>& InRate);
+};
