DaySequence: Move DaySequence from Restricted/NFL to Experimental

diff --git a/Engine/Plugins/Experimental/DaySequence/Source/DaySequence/Private/DaySequenceModifierComponent.cpp b/Engine/Plugins/Experimental/DaySequence/Source/DaySequence/Private/DaySequenceModifierComponent.cpp
new file mode 100644
--- /dev/null
+++ b/Engine/Plugins/Experimental/DaySequence/Source/DaySequence/Private/DaySequenceModifierComponent.cpp
@@ -0,0 +1,1587 @@
+// Copyright Epic Games, Inc. All Rights Reserved.
+
+#include "DaySequenceModifierComponent.h"
+
+#include "Components/BoxComponent.h"
+#include "DaySequence.h"
+#include "DaySequenceCollectionAsset.h"
+#include "DaySequenceTrack.h"
+
+#include "Engine/World.h"
+#include "Materials/MaterialInterface.h"
+
+#include "Components/CapsuleComponent.h"
+#include "Components/SphereComponent.h"
+
+#include "MovieScene.h"
+#include "MovieSceneTrack.h"
+#include "MovieSceneSection.h"
+#include "MovieSceneBinding.h"
+#include "MovieScenePossessable.h"
+
+#include "Tracks/MovieSceneSubTrack.h"
+#include "Tracks/MovieSceneBoolTrack.h"
+#include "Tracks/MovieSceneFloatTrack.h"
+#include "Tracks/MovieSceneColorTrack.h"
+#include "Tracks/MovieSceneDoubleTrack.h"
+#include "Tracks/MovieSceneVectorTrack.h"
+#include "Tracks/MovieSceneMaterialTrack.h"
+#include "Tracks/MovieSceneVisibilityTrack.h"
+#include "Tracks/MovieScene3DTransformTrack.h"
+#include "Tracks/MovieScenePrimitiveMaterialTrack.h"
+
+#include "Sections/MovieSceneSubSection.h"
+#include "Sections/MovieSceneBoolSection.h"
+#include "Sections/MovieSceneFloatSection.h"
+#include "Sections/MovieSceneColorSection.h"
+#include "Sections/MovieSceneDoubleSection.h"
+#include "Sections/MovieSceneVectorSection.h"
+#include "Sections/MovieScene3DTransformSection.h"
+#include "Sections/MovieScenePrimitiveMaterialSection.h"
+
+#include "ProfilingDebugging/CsvProfiler.h"
+
+#if ENABLE_DRAW_DEBUG
+#endif
+
+#define LOCTEXT_NAMESPACE "DaySequenceModifierComponent"
+
+namespace UE::DaySequence
+{
+	template<typename TrackType>
+	TrackType* CreateOrAddOverrideTrack(UMovieScene* MovieScene, const FGuid& ObjectGuid, FName Name = NAME_None)
+	{
+		TrackType* Track = MovieScene->FindTrack<TrackType>(ObjectGuid, Name);
+		if (!Track)
+		{
+			// Clear RF_Transactional and set RF_Transient on created tracks and sections
+			// to avoid dirtying the package for these procedurally generated sequences.
+			// RF_Transactional is explicitly set in UMovieSceneSection/Track::PostInitProperties.
+			Track = NewObject<TrackType>(MovieScene, NAME_None, RF_Transient);
+			Track->ClearFlags(RF_Transactional);
+
+			UMovieSceneSection* Section = Track->CreateNewSection();
+			Section->ClearFlags(RF_Transactional);
+			Section->SetFlags(RF_Transient);
+			Section->SetRange(TRange<FFrameNumber>::All());
+
+			Track->AddSection(*Section);
+			MovieScene->AddGivenTrack(Track, ObjectGuid);
+		}
+
+		return Track;
+	}
+
+	template<typename TrackType>
+	TrackType* CreateOrAddPropertyOverrideTrack(UMovieScene* MovieScene, const FGuid& ObjectGuid, FName InPropertyName)
+	{
+		TrackType* Track = CreateOrAddOverrideTrack<TrackType>(MovieScene, ObjectGuid, InPropertyName);
+		check(Track);
+		
+		const FString PropertyPath = InPropertyName.ToString();
+
+		// Split the property path to capture the leaf property name and parent struct to conform
+		// with Sequencer Editor property name/path and display name conventions:
+		//
+		// PropertyName = MyProperty
+		// PropertyPath = MyPropertyStruct.MyProperty
+		// DisplayName = PropertyName (PropertyStruct)
+		FName PropertyName;
+		FName PropertyParent;
+		int32 NamePos = INDEX_NONE;
+		if (PropertyPath.FindLastChar('.', NamePos) && NamePos < PropertyPath.Len() - 1)
+		{
+			PropertyName = FName(FStringView(*PropertyPath + NamePos + 1, PropertyPath.Len() - NamePos - 1));
+			PropertyParent = FName(FStringView(*PropertyPath, NamePos));
+		}
+		else
+		{
+			PropertyName = *PropertyPath;
+		}
+		
+		Track->SetPropertyNameAndPath(PropertyName, PropertyPath);
+
+#if WITH_EDITOR
+		if (NamePos != INDEX_NONE)
+		{
+			FText DisplayText = FText::Format(LOCTEXT("ModifierPropertyTrackFormat", "{0} ({1})"), FText::FromName(PropertyName), FText::FromName(PropertyParent));
+			Track->SetDisplayName(DisplayText);
+		}
+#endif
+		return Track;
+	}
+
+	template<typename TrackType, typename SectionType>
+	SectionType* CreateOrAddPropertyOverrideSection(UMovieScene* MovieScene, const FGuid& ObjectGuid, FName PropertyName)
+	{
+		TrackType* Track = CreateOrAddPropertyOverrideTrack<TrackType>(MovieScene, ObjectGuid, PropertyName);
+		check(Track);
+		return Cast<SectionType>(Track->GetAllSections()[0]);
+	}
+
+	template<typename TrackType, typename SectionType>
+	SectionType* CreateOrAddOverrideSection(UMovieScene* MovieScene, const FGuid& ObjectGuid)
+	{
+		TrackType* Track = CreateOrAddOverrideTrack<TrackType>(MovieScene, ObjectGuid);
+		check(Track);
+		return Cast<SectionType>(Track->GetAllSections()[0]);
+	}
+
+	FVector GVolumePreviewLocation = FVector::ZeroVector;
+
+	float ComputeBoxSignedDistance(const UBoxComponent* BoxComponent, const FVector& InWorldPosition)
+	{
+		const FTransform& ComponentTransform = BoxComponent->GetComponentTransform();
+
+		FVector Point = ComponentTransform.InverseTransformPositionNoScale(InWorldPosition);
+		FVector Box = BoxComponent->GetUnscaledBoxExtent() * ComponentTransform.GetScale3D();
+
+		FVector Delta = Point.GetAbs() - Box;
+		return FVector::Max(Delta, FVector::ZeroVector).Length() + FMath::Min(Delta.GetMax(), 0.0);
+	}
+
+	float ComputeSphereSignedDistance(const USphereComponent* SphereComponent, const FVector& InWorldPosition)
+	{
+		const FTransform& ComponentTransform = SphereComponent->GetComponentTransform();
+
+		FVector Point = ComponentTransform.InverseTransformPositionNoScale(InWorldPosition);
+		return Point.Length()-SphereComponent->GetScaledSphereRadius();
+	}
+	
+	float ComputeCapsuleSignedDistance(const UCapsuleComponent* CapsuleComponent, const FVector& InWorldPosition)
+	{
+		// UCapsuleComponent::GetScaledCapsuleRadius() returns the min scaled X/Y axis for the radius
+		// while the actual collision query uses the max scaled X/Y axis for the radius. We use Max here to match the collision.
+		const FTransform& ComponentTransform = CapsuleComponent->GetComponentTransform();
+		const FVector&    ComponentScale     = ComponentTransform.GetScale3D();
+
+		FVector Point = ComponentTransform.InverseTransformPositionNoScale(InWorldPosition);
+		const double CapsuleHalfHeight = CapsuleComponent->GetScaledCapsuleHalfHeight_WithoutHemisphere();
+		const double CapsuleRadius = CapsuleComponent->GetUnscaledCapsuleRadius() * FMath::Max(ComponentScale.X, ComponentScale.Y);
+
+		Point.Z = FMath::Max(FMath::Abs(Point.Z) - CapsuleHalfHeight, 0.0);
+		return Point.Length() - CapsuleRadius;
+	}
+
+	float ComputeSignedDistance(const UShapeComponent* ShapeComponent, const FVector& InWorldPosition)
+	{
+		if (!ShapeComponent)
+		{
+			return UE_MAX_FLT;
+		}
+		
+		if (const UBoxComponent* BoxComponent = Cast<UBoxComponent>(ShapeComponent))
+		{
+			return ComputeBoxSignedDistance(BoxComponent, InWorldPosition);
+		}
+		else if (const USphereComponent* SphereComponent = Cast<USphereComponent>(ShapeComponent))
+		{
+			return ComputeSphereSignedDistance(SphereComponent, InWorldPosition);
+		}
+		else if (const UCapsuleComponent* CapsuleComponent = Cast<UCapsuleComponent>(ShapeComponent))
+		{
+			return ComputeCapsuleSignedDistance(CapsuleComponent, InWorldPosition);
+		}
+		
+		// @todo: unsupported shape?
+		return (InWorldPosition - ShapeComponent->GetComponentLocation()).Length();
+	}
+
+	bool TestValidProperty(UObject* Object, FProperty* Property)
+	{
+		if (!Property)
+		{
+			FFrame::KismetExecutionMessage(*FString::Printf(TEXT("Invalid property specified for object %s."), *Object->GetName()), ELogVerbosity::Error);
+			return false;
+		}
+
+		if (Property->HasAnyPropertyFlags(CPF_Deprecated))
+		{
+			// Emit a warning for deprecated properties but still consider them valid
+			FFrame::KismetExecutionMessage(*FString::Printf(TEXT("Depcrecated property specified: %s for object %s."), *Property->GetName(), *Object->GetName()), ELogVerbosity::Warning);
+		}
+
+		return true;
+	}
+
+} // namespace UE::DaySequence
+
+void UDaySequenceModifierEasingFunction::Initialize(EEasingFunctionType EasingType)
+{
+	if (UDaySequenceModifierComponent* Outer = Cast<UDaySequenceModifierComponent>(GetOuter()))
+	{
+		switch (EasingType)
+		{
+		case EEasingFunctionType::EaseIn:
+			EvaluateImpl = [Outer](float)
+			{
+				return Outer->GetCurrentBlendWeight();
+			};
+			break;
+	
+		case EEasingFunctionType::EaseOut:
+			EvaluateImpl = [Outer](float)
+			{
+				return 1.f - Outer->GetCurrentBlendWeight();
+			};
+			break;
+		}
+	}
+	else
+	{
+		ensureMsgf(false, TEXT("Initialize called without a valid Outer!"));
+		EvaluateImpl = [](float){ return 0.f; };
+	}
+}
+
+float UDaySequenceModifierEasingFunction::Evaluate(float Interp) const
+{
+	return EvaluateImpl(Interp);
+}
+
+UDaySequenceModifierComponent::UDaySequenceModifierComponent(const FObjectInitializer& Init)
+	: Super(Init)
+{
+	bIsComponentEnabled = true;
+	bIsEnabled = false;
+	bIgnoreBias = false;
+	bUseVolume = true;
+	bPreview = true;
+	bUseCollection = false;
+	bCachedExternalShapesInvalid = true;
+	Bias = 1000;
+	DayNightCycleTime = 12.f;
+	DayNightCycle = EDayNightCycleMode::Default;
+	BlendMode = EDaySequenceModifierBlendMode::Distance;
+	BlendAmount = 100.f;
+	CachedBlendFactor = 0.f;
+	CustomVolumeBlendWeight = 1.0f;
+#if ENABLE_DRAW_DEBUG
+	DebugLevel = 0;
+	
+	// This gets captured by a lambda below so should continue living
+	TSharedPtr<TMap<FString, FString>> DebugData = MakeShared<TMap<FString, FString>>();
+	DebugEntry = MakeShared<UE::DaySequence::FDaySequenceDebugEntry>(
+	[this](){ return ShouldShowDebugInfo(); },
+	[this, DebugData]()
+	{
+		(*DebugData).FindOrAdd("Owner Name") = GetOwner()->GetFName().ToString();
+		(*DebugData).FindOrAdd("Component Enabled") = bIsComponentEnabled ? "True" : "False";
+		(*DebugData).FindOrAdd("Modifier Enabled") = bIsEnabled ? "True" : "False";
+		(*DebugData).FindOrAdd("Blend Weight") = FString::Printf(TEXT("%.5f"), GetCurrentBlendWeight());
+		(*DebugData).FindOrAdd("Blend Target" ) = ExternalVolumeBlendTarget.IsValid() ? ExternalVolumeBlendTarget->GetName() : "None";
+
+		return DebugData;
+	});
+#endif
+	
+	PrimaryComponentTick.bCanEverTick = true;
+	PrimaryComponentTick.bStartWithTickEnabled = false;
+	PrimaryComponentTick.TickGroup = TG_PostPhysics;
+	
+	EasingFunction = CreateDefaultSubobject<UDaySequenceModifierEasingFunction>("EasingFunction", true);
+}
+
+#if WITH_EDITOR
+
+void UDaySequenceModifierComponent::SetVolumePreviewLocation(const FVector& Location)
+{
+	UE::DaySequence::GVolumePreviewLocation = Location;
+}
+
+void UDaySequenceModifierComponent::UpdateEditorPreview(float DeltaTime)
+{
+	using namespace UE::DaySequence;
+
+	if (bIsComponentEnabled && bPreview && bUseVolume && IsRegistered() && !GetWorld()->IsGameWorld())
+	{
+		if (const float DistanceBlendFactor = GetDistanceBlendFactor(GVolumePreviewLocation); DistanceBlendFactor > UE_SMALL_NUMBER)
+		{
+			if (!bIsEnabled)
+			{
+				EnableModifier();
+			}
+			
+			const float BlendFactor = FMath::Min(DistanceBlendFactor, CustomVolumeBlendWeight);
+			if (BlendMode != EDaySequenceModifierBlendMode::None && CachedBlendFactor != BlendFactor)
+			{
+				// If we're using a blend we have to mark active sections as changed
+				// in order to force an update in-editor:
+				
+				if (UMovieSceneSubSection* SubSection = WeakSubSection.Get())
+				{
+					SubSection->MarkAsChanged();
+				}
+
+				for (TWeakObjectPtr<UMovieSceneSubSection> SubSection : SubSections)
+				{
+					UMovieSceneSubSection* StrongSubSection = SubSection.Get();
+					if (StrongSubSection && StrongSubSection->IsActive())
+					{
+						StrongSubSection->MarkAsChanged();
+						break;
+					}
+				}
+			}
+			CachedBlendFactor = BlendFactor;
+		}
+		else
+		{
+			DisableModifier();
+		}
+	}
+}
+
+TStatId UDaySequenceModifierComponent::GetStatId() const
+{
+	RETURN_QUICK_DECLARE_CYCLE_STAT(UDaySequenceModifierComponent, STATGROUP_Tickables);
+}
+
+ETickableTickType UDaySequenceModifierComponent::GetTickableTickType() const
+{
+	UWorld* World = GetWorld();
+	if (World && World->WorldType == EWorldType::Editor)
+	{
+		return ETickableTickType::Always;
+	}
+	return ETickableTickType::Never;
+}
+
+void UDaySequenceModifierComponent::PostEditChangeProperty(FPropertyChangedEvent& PropertyChangedEvent)
+{
+	Super::PostEditChangeProperty(PropertyChangedEvent);
+	
+	const FName PropertyName = PropertyChangedEvent.Property ? PropertyChangedEvent.Property->GetFName() : NAME_None;
+
+	if (PropertyName == GET_MEMBER_NAME_CHECKED(UDaySequenceModifierComponent, bPreview))
+	{
+		if (bPreview && !bIsEnabled)
+		{
+			EnableModifier();
+		}
+		else if (!bPreview && bIsEnabled)
+		{
+			DisableModifier();
+		}
+	}
+	else if (PropertyName == GET_MEMBER_NAME_CHECKED(UDaySequenceModifierComponent, bUseVolume))
+	{
+		if (!bUseVolume)
+		{
+			if (bPreview && !bIsEnabled)
+			{
+				EnableModifier();
+			}
+			else if (!bPreview && bIsEnabled)
+			{
+				DisableModifier();
+			}
+		}
+	}
+}
+
+#endif // WITH_EDITOR
+
+void UDaySequenceModifierComponent::OnRegister()
+{
+	Super::OnRegister();
+
+	bCachedExternalShapesInvalid = true;
+	
+	if (!bUseVolume)
+	{
+		EnableModifier();
+	}
+}
+
+void UDaySequenceModifierComponent::OnUnregister()
+{
+	Super::OnUnregister();
+
+	bCachedExternalShapesInvalid = true;
+	
+	DisableModifier();
+	RemoveSubSequenceTrack();
+}
+
+void UDaySequenceModifierComponent::BeginPlay()
+{
+	Super::BeginPlay();
+	
+	BindOverlapEvents();
+}
+
+void UDaySequenceModifierComponent::TickComponent(float DeltaTime, ELevelTick TickType, FActorComponentTickFunction* ThisTickFunction)
+{
+	Super::TickComponent(DeltaTime, TickType, ThisTickFunction);
+
+	// If we don't need to tick, don't tick
+	AActor* BlendTarget = ExternalVolumeBlendTarget.Get();
+	if (!BlendTarget || !TargetActor || TargetActor->IsPlaying())
+	{
+		PrimaryComponentTick.SetTickFunctionEnable(false);
+		return;
+	}
+
+	// Update TargetActor if our blend factor has changed
+	const float DistanceBlendFactor = GetDistanceBlendFactor(TargetActor->GetActorLocation());
+	const float BlendFactor = FMath::Min(DistanceBlendFactor, CustomVolumeBlendWeight);
+	if (DistanceBlendFactor > UE_SMALL_NUMBER && CachedBlendFactor != BlendFactor)
+	{
+		TargetActor->SetTimeOfDay(TargetActor->GetTimeOfDay());
+	}
+	CachedBlendFactor = BlendFactor;
+}
+
+void UDaySequenceModifierComponent::EndPlay(EEndPlayReason::Type Reason)
+{
+	Super::EndPlay(Reason);
+	
+	UnbindOverlapEvents();
+	RemoveSubSequenceTrack();
+}
+
+void UDaySequenceModifierComponent::ResetOverrides()
+{
+	if (ProceduralDaySequence)
+	{
+		UMovieScene* MovieScene = ProceduralDaySequence->GetMovieScene();
+		TArray<FMovieSceneBinding> Bindings = MovieScene->GetBindings();
+
+		for (const FMovieSceneBinding& Binding : Bindings)
+		{
+			FGuid BindingID = Binding.GetObjectGuid();
+
+			ProceduralDaySequence->UnbindPossessableObjects(BindingID);
+			MovieScene->RemovePossessable(BindingID);
+		}
+	}
+
+	if (TargetActor && bUnpauseOnDisable)
+	{
+		TargetActor->Play();
+		bUnpauseOnDisable = false;
+	}
+}
+
+void UDaySequenceModifierComponent::BindToDaySequenceActor(ADaySequenceActor* DaySequenceActor)
+{
+	if (TargetActor == DaySequenceActor)
+	{
+		return;
+	}
+
+	bool bWasEnabled = bIsEnabled;
+	UnbindFromDaySequenceActor();
+
+	TargetActor = DaySequenceActor;
+
+	if (bWasEnabled)
+	{
+		EnableModifier();
+	}
+
+	if (ensureMsgf(DaySequenceActor, TEXT("BindToDaySequenceActor called with a null Day Sequence Actor.")))
+	{
+		DaySequenceActor->GetOnPostInitializeDaySequences().AddUObject(this, &UDaySequenceModifierComponent::ReinitializeSubSequence);
+#if ENABLE_DRAW_DEBUG
+		if (!DaySequenceActor->IsDebugCategoryRegistered(ShowDebug_ModifierCategory))
+		{
+			DaySequenceActor->RegisterDebugCategory(ShowDebug_ModifierCategory, TargetActor->OnShowDebugInfoDrawFunction);
+		}
+		
+		DaySequenceActor->GetOnDebugLevelChanged().AddUObject(this, &UDaySequenceModifierComponent::OnDebugLevelChanged);
+		DaySequenceActor->RegisterDebugEntry(DebugEntry, ShowDebug_ModifierCategory);
+#endif
+	}
+}
+
+void UDaySequenceModifierComponent::UnbindFromDaySequenceActor()
+{
+	DisableModifier();
+	RemoveSubSequenceTrack();
+
+	if (TargetActor)
+	{
+		TargetActor->GetOnPostInitializeDaySequences().RemoveAll(this);
+#if ENABLE_DRAW_DEBUG
+		TargetActor->GetOnDebugLevelChanged().RemoveAll(this);
+		TargetActor->UnregisterDebugEntry(DebugEntry, ShowDebug_ModifierCategory);
+#endif
+		TargetActor = nullptr;
+	}
+}
+
+void UDaySequenceModifierComponent::RemoveSubSequenceTrack()
+{
+	auto RemoveSubTrack = [](const UMovieSceneSubSection* SubSection)
+	{
+		if (SubSection)
+		{
+			UMovieSceneTrack* Track = SubSection->GetTypedOuter<UMovieSceneTrack>();
+			UMovieScene* MovieScene = Track->GetTypedOuter<UMovieScene>();
+
+			check(Track && MovieScene);
+
+			MovieScene->RemoveTrack(*Track);
+			MovieScene->MarkAsChanged();
+		}
+	};
+	
+	RemoveSubTrack(WeakSubSection.Get());
+	WeakSubSection = nullptr;
+
+	for (const TWeakObjectPtr<UMovieSceneSubSection> SubSection : SubSections)
+	{
+		RemoveSubTrack(SubSection.Get());
+	}
+	SubSections.Empty();
+
+#if ENABLE_DRAW_DEBUG
+	if (TargetActor)
+	{
+		for (TSharedPtr<UE::DaySequence::FDaySequenceDebugEntry> Entry : SubSectionDebugEntries)
+		{
+			TargetActor->UnregisterDebugEntry(Entry, TargetActor->ShowDebug_SubSequenceCategory);
+		}
+	}
+	SubSectionDebugEntries.Empty();
+#endif
+}
+
+bool UDaySequenceModifierComponent::CanBeEnabled() const
+{
+	AActor* Actor = TargetActor ? TargetActor : GetOwner();
+
+	if (!bIsComponentEnabled)
+	{
+		return false;
+	}
+	
+	if (bUseVolume)
+	{
+		ENetMode NetMode = Actor->GetNetMode();
+		return NetMode == NM_Standalone || NetMode == NM_Client;
+	}
+
+	return true;
+}
+
+void UDaySequenceModifierComponent::EnableComponent()
+{
+	if (bIsComponentEnabled)
+	{
+		return;
+	}
+	
+	bIsComponentEnabled = true;
+}
+
+void UDaySequenceModifierComponent::DisableComponent()
+{
+	if (!bIsComponentEnabled && !bIsEnabled)
+	{
+		return;
+	}
+	
+	bIsComponentEnabled = false;
+
+	DisableModifier();
+	RemoveSubSequenceTrack();
+}
+
+void UDaySequenceModifierComponent::EnableModifier()
+{
+	CSV_SCOPED_TIMING_STAT_EXCLUSIVE(DaySequenceModifierComponent_EnableModifier);
+	
+	if (bIsEnabled || !CanBeEnabled())
+	{
+		return;
+	}
+
+	if (!bPreview && GetWorld()->WorldType == EWorldType::Editor)
+	{
+		return;
+	}
+
+	bIsEnabled = true;
+
+	// Will call SetSubTrackMuteState for all living subsections, which checks enable state of modifier and their conditions
+	InvalidateMuteStates();
+
+	// In both collection and non collection case this array is populated, so if size is 0 we never initialized or removed subsections
+	if (SubSections.Num() == 0)
+	{
+		ReinitializeSubSequence(nullptr);
+	}
+
+	SetInitialTimeOfDay();
+
+	// Force an update if it's not playing so that the effects of this being enabled are seen
+	if (TargetActor && !TargetActor->IsPlaying())
+	{
+		TargetActor->SetTimeOfDay(TargetActor->GetTimeOfDay());
+
+		// If the actor is not playing and we are using a blend target, enable ticking on this component
+		// since it is needed to update the target actor if the blend amount changes
+		if (ExternalVolumeBlendTarget.IsValid())
+		{
+			PrimaryComponentTick.SetTickFunctionEnable(true);
+		}
+	}
+
+	OnPostEnableModifier.Broadcast();
+}
+
+void UDaySequenceModifierComponent::DisableModifier()
+{
+	CSV_SCOPED_TIMING_STAT_EXCLUSIVE(DaySequenceModifierComponent_DisableModifier);
+	
+	if (!bIsEnabled)
+	{
+		return;
+	}
+
+	if (!bPreview && GetWorld()->WorldType == EWorldType::Editor)
+	{
+		return;
+	}
+	
+	bIsEnabled = false;
+	PrimaryComponentTick.SetTickFunctionEnable(false);
+
+	if (TargetActor && !TargetActor->HasAnyFlags(RF_BeginDestroyed))
+	{
+		// Will call SetSubTrackMuteState for all living subsections, which checks enable state of modifier and their conditions
+		InvalidateMuteStates();
+
+		TargetActor->RemoveStaticTimeOfDay();
+		
+		if (bUnpauseOnDisable)
+		{
+			TargetActor->Play();
+			bUnpauseOnDisable = false;
+		}
+		// Force an update if it's not playing so that the effects of this being disabled are seen
+		else if (!TargetActor->IsPlaying())
+		{
+			TargetActor->SetTimeOfDay(TargetActor->GetTimeOfDay());
+		}
+	}
+
+	// Necessary for correctly marking TargetActor as changed on enable.
+	CachedBlendFactor = -1.0f;
+}
+
+void UDaySequenceModifierComponent::SetInitialTimeOfDay()
+{
+	if (TargetActor)
+	{
+		const bool bHasAuthority = TargetActor->HasAuthority();
+		const bool bRandomTimeOfDay = DayNightCycle == EDayNightCycleMode::RandomFixedTime || DayNightCycle == EDayNightCycleMode::RandomStartTime;
+		const float Time = bRandomTimeOfDay ? FMath::FRand()*TargetActor->GetDayLength() : DayNightCycleTime;
+
+		if (!bHasAuthority && !bUseVolume)
+		{
+			// Never set initial time of day from non-volume based modifiers if they don't have authority
+			// We'll just get the initial time of day from the server replication
+			return;
+		}
+
+		switch (DayNightCycle)
+		{
+		case EDayNightCycleMode::FixedTime:
+		case EDayNightCycleMode::RandomFixedTime:
+			if (!bHasAuthority && bUseVolume)
+			{
+				// This function assigns a custom time controller so we can override
+				// the time regardless of server replication
+				TargetActor->SetStaticTimeOfDay(Time);
+				break;
+			}
+
+			// If we're not overriding the time on a client, we need to make sure the time is replicated correctly
+			// AddStaticTimeOfDayOverride should be used where a static time of day needs to be evaluated from the sequence itself (to support h-bias overriding)
+			bUnpauseOnDisable = TargetActor->IsPlaying();
+			TargetActor->Pause();
+
+			// Intentional fallthrough - set the time and preview time
+
+		case EDayNightCycleMode::StartAtSpecifiedTime:
+		case EDayNightCycleMode::RandomStartTime:
+
+			TargetActor->SetTimeOfDay(Time);
+#if WITH_EDITOR
+			TargetActor->ConditionalSetTimeOfDayPreview(Time);
+#endif
+			break;
+
+#if WITH_EDITOR
+		default:
+			TargetActor->SetTimeOfDayPreview(TargetActor->GetTimeOfDayPreview());
+#endif
+		}
+	}
+}
+
+void UDaySequenceModifierComponent::ReinitializeSubSequence(ADaySequenceActor::FSubSectionPreserveMap* SectionsToPreserve)
+{
+#if ROOT_SEQUENCE_RECONSTRUCTION_ENABLED
+	bool bReinit = true;
+
+	if (SectionsToPreserve)
+	{
+		// Mark all subsections we have recorded for keep in the root sequence
+		// This is a fast path we take only if all of our subsections are in the root sequence
+		for (TWeakObjectPtr<UMovieSceneSubSection> SubSection : SubSections)
+		{
+			if (const UMovieSceneSubSection* StrongSubSection = SubSection.Get())
+			{
+				if (bool* SectionToPreserveFlag = SectionsToPreserve->Find(StrongSubSection))
+				{
+					*SectionToPreserveFlag = true;
+					bReinit = false;
+				}
+				else
+				{
+					// If we have a subsection that is not in the root sequence, break and reinit completely
+					bReinit = true;
+					break;
+				}
+			}
+		}
+
+		if (bReinit)
+		{
+			// Mark all sections associated with this modifier for delete before we do a full reinit
+			for (TWeakObjectPtr<UMovieSceneSubSection> SubSection : SubSections)
+			{
+				if (const UMovieSceneSubSection* StrongSubSection = SubSection.Get())
+				{
+					if (bool* SectionToPreserveFlag = SectionsToPreserve->Find(StrongSubSection))
+					{
+						*SectionToPreserveFlag = false;
+					}
+				}
+			}
+		}
+	}
+
+	if (bReinit)
+	{
+#endif
+		RemoveSubSequenceTrack();
+		
+		if (bUseCollection)
+		{
+			if (DaySequenceCollection)
+			{
+				for (const FDaySequenceCollectionEntry& Entry : DaySequenceCollection->DaySequences)
+				{
+					InitializeDaySequence(Entry);
+				}
+			}
+		}
+		else
+		{
+			// Always create the sub section even if it is nullptr. This means that
+			// if the procedural day sequence is created later, we can still add it to the sub section
+			UDaySequence* SequenceToUse = UserDaySequence ? UserDaySequence : ProceduralDaySequence;
+	
+			WeakSubSection = InitializeDaySequence(SequenceToUse);
+		}
+#if ROOT_SEQUENCE_RECONSTRUCTION_ENABLED
+	}
+	else
+	{
+		// If we took the fast path, invalidate all mute states.
+		InvalidateMuteStates();
+	}
+#endif
+	
+#if ENABLE_DRAW_DEBUG
+	if (TargetActor)
+	{
+		if (!TargetActor->IsDebugCategoryRegistered(TargetActor->ShowDebug_SubSequenceCategory))
+		{
+			TargetActor->RegisterDebugCategory(TargetActor->ShowDebug_SubSequenceCategory,  TargetActor->OnShowDebugInfoDrawFunction);
+		}
+	
+		for (TSharedPtr<UE::DaySequence::FDaySequenceDebugEntry> Entry : SubSectionDebugEntries)
+		{
+			TargetActor->RegisterDebugEntry(Entry, TargetActor->ShowDebug_SubSequenceCategory);
+		}
+	}
+#endif
+	
+	OnPostReinitializeSubSequences.Broadcast();
+}
+
+UMovieSceneSubSection* UDaySequenceModifierComponent::InitializeDaySequence(const FDaySequenceCollectionEntry& Entry)
+{
+	UDaySequence* RootSequence = TargetActor ? TargetActor->GetRootSequence() : nullptr;
+    UMovieScene*  MovieScene   = RootSequence ? RootSequence->GetMovieScene() : nullptr;
+
+    if (!MovieScene)
+    {
+    	return nullptr;
+    }
+	
+	auto CreateSubTrack = [this, MovieScene](UDaySequence* Sequence, int BiasOffset, bool bActivate, bool bBlendHierarchicalBias)
+	{
+		UDaySequenceTrack*     RootTrack  = MovieScene->AddTrack<UDaySequenceTrack>();
+		RootTrack->ClearFlags(RF_Transactional);
+		RootTrack->SetFlags(RF_Transient);
+        
+		UMovieSceneSubSection* SubSection = CastChecked<UMovieSceneSubSection>(RootTrack->CreateNewSection());
+		SubSection->ClearFlags(RF_Transactional);
+		// SubSections of DaySequenceTrack will inherit flags from its parent track - RF_Transient in this case.
+		SubSection->Parameters.HierarchicalBias = Bias + BiasOffset;
+		SubSection->Parameters.Flags            = EMovieSceneSubSectionFlags::OverrideRestoreState | (bIgnoreBias ? EMovieSceneSubSectionFlags::IgnoreHierarchicalBias : EMovieSceneSubSectionFlags::None);
+
+		SubSection->SetSequence(Sequence);
+		SubSection->SetRange(MovieScene->GetPlaybackRange());
+		SubSection->SetIsActive(bActivate);
+
+		TargetActor->UpdateSubSectionTimeScale(SubSection);
+
+		RootTrack->AddSection(*SubSection);
+
+		if (bUseVolume && BlendMode != EDaySequenceModifierBlendMode::None && bBlendHierarchicalBias)
+		{
+			// In the Sequencer Editor, EaseIn pads the Sequence asset name by the EaseIn duration
+			// (see SSequencerSection::OnPaint). Since we set the Easing duration to the full section
+			// width to facilitate blending, the label is clipped. So we use EaseOut here instead and
+			// ensure that the weight is inverted in Evaluate().
+			SubSection->Easing.bManualEaseOut = true;
+			SubSection->Easing.ManualEaseOutDuration = MovieScene->GetPlaybackRange().Size<FFrameNumber>().Value;
+
+			EasingFunction->Initialize(UDaySequenceModifierEasingFunction::EEasingFunctionType::EaseOut);
+			SubSection->Easing.EaseOut = EasingFunction;
+			SubSection->Parameters.Flags |= EMovieSceneSubSectionFlags::BlendHierarchicalBias;
+		}
+
+#if WITH_EDITOR
+		FString Label = GetOwner()->GetActorLabel();
+#else
+		FString Label = GetOwner()->GetName();
+#endif
+#if WITH_EDITORONLY_DATA
+		RootTrack->DisplayName = FText::Format(LOCTEXT("ModifierTrackFormat", "Modifier ({0})"), FText::FromString(Label));
+#endif
+        
+		SubSection->MarkAsChanged();
+		return SubSection;
+	};
+	
+	constexpr bool bActivate = true;
+	constexpr bool bBlendHierarchicalBias = true;
+	UMovieSceneSubSection* SubSection = CreateSubTrack(Entry.Sequence, Entry.BiasOffset, bActivate, bBlendHierarchicalBias);
+	
+	if (!SubSections.Contains(SubSection))
+	{
+		SubSections.Add(SubSection);
+	}
+
+	const TFunction<void(void)> SetSubTrackMuteStateConditional = [this, SubSection, Conditions = Entry.Conditions.Conditions]()
+	{
+		if (!IsValid(this) || !IsValid(SubSection))
+		{
+			return;
+		}
+
+		constexpr bool bInitialMuteState = false;
+		const bool bActive = bIsEnabled && !TargetActor->EvaluateSequenceConditions(bInitialMuteState, Conditions);
+		if (SubSection->IsActive() != bActive)
+		{
+			SubSection->MarkAsChanged();
+			SubSection->SetIsActive(bActive);
+		}
+	};
+
+	const TFunction<void(void)> SetSubTrackMuteStateUnconditional = [this, SubSection]()
+	{
+		if (!IsValid(this) || !IsValid(SubSection))
+		{
+			return;
+		}
+
+		const bool bActive = bIsEnabled;
+		if (SubSection->IsActive() != bActive)
+		{
+			SubSection->MarkAsChanged();
+			SubSection->SetIsActive(bActive);
+		}
+	};
+
+	const TFunction<void(void)>& SetSubTrackMuteState = Entry.Conditions.Conditions.Num() == 0 ? SetSubTrackMuteStateUnconditional : SetSubTrackMuteStateConditional;
+	
+	// Initialize mute state and set up the condition callbacks to dynamically update mute state.
+	SetSubTrackMuteState();
+	OnInvalidateMuteStates.AddWeakLambda(SubSection, SetSubTrackMuteState);
+	TargetActor->BindToConditionCallbacks(this, Entry.Conditions.Conditions, [this]() { InvalidateMuteStates(); });
+	
+#if ENABLE_DRAW_DEBUG
+	// This gets captured by a lambda below so should continue living
+	TSharedPtr<TMap<FString, FString>> DebugData = MakeShared<TMap<FString, FString>>();
+	SubSectionDebugEntries.Emplace(MakeShared<UE::DaySequence::FDaySequenceDebugEntry>(
+	[this](){ return true; },
+	[this, DebugData, SubSection]()
+	{
+		if (IsValid(SubSection))
+		{
+			(*DebugData).FindOrAdd("Owner Name") = GetOwner()->GetFName().ToString();
+			(*DebugData).FindOrAdd("Sequence Name") = SubSection->GetSequence() ? SubSection->GetSequence()->GetFName().ToString() : "None";
+			(*DebugData).FindOrAdd("Mute State") = SubSection->IsActive() ? "Active" : "Muted";
+			(*DebugData).FindOrAdd("Hierarchical Bias") = FString::Printf(TEXT("%d"), SubSection->Parameters.HierarchicalBias);
+		}
+	
+		return DebugData;
+	}));
+#endif
+
+	return SubSection;
+}
+
+FGuid UDaySequenceModifierComponent::GetOrCreateProceduralBinding(UObject* Object)
+{
+	if (!Object)
+	{
+		FFrame::KismetExecutionMessage(TEXT("Null Object parameter specified."), ELogVerbosity::Error);
+		return FGuid();
+	}
+
+	USceneComponent* Component = Cast<USceneComponent>(Object);
+	AActor*          Actor     = Cast<AActor>(Object);
+
+	// Set up the time of day actor binding if we don't have one already
+	if (!TargetActor)
+	{
+		if (Actor)
+		{
+			TargetActor = Cast<ADaySequenceActor>(Actor);
+		}
+		else if (Component)
+		{
+			TargetActor = Cast<ADaySequenceActor>(Component->GetOwner());
+		}
+	}
+
+	if (!TargetActor)
+	{
+		FFrame::KismetExecutionMessage(TEXT("No valid ADaySequenceActor has been set up. Have you called BindToDaySequenceActor yet?"), ELogVerbosity::Error);
+		return FGuid();
+	}
+
+	if (Component && !Component->IsIn(TargetActor))
+	{
+		FFrame::KismetExecutionMessage(TEXT("Unable to bind to components that exist outside of the ADaySequenceActor we are tracking."), ELogVerbosity::Error);
+		return FGuid();
+	}
+
+	if (Actor && Actor != TargetActor)
+	{
+		FFrame::KismetExecutionMessage(TEXT("Unable to bind to actors that are not the ADaySequenceActor we are tracking."), ELogVerbosity::Error);
+		return FGuid();
+	}
+
+	if (!ProceduralDaySequence)
+	{
+		// Name the procedural sequence the same as this component's owner so it shows up in the Sequence
+		// with a meaningful name
+#if WITH_EDITOR
+		FName SequenceName = MakeUniqueObjectName(this, UDaySequence::StaticClass(), *GetOwner()->GetActorLabel());
+#else
+		FName SequenceName = MakeUniqueObjectName(this, UDaySequence::StaticClass(), GetOwner()->GetFName());
+#endif
+
+		ProceduralDaySequence = NewObject<UDaySequence>(this, SequenceName);
+		ProceduralDaySequence->Initialize(RF_Transient);
+	}
+
+	// If we have a sub section but it has no sequence applied, apply it now.
+	// This implies EnableModifier was called before we had any valid sequence data
+	UMovieSceneSubSection* SubSection = WeakSubSection.Get();
+	if (SubSection && SubSection->GetSequence() == nullptr)
+	{
+		SubSection->MarkAsChanged();
+		SubSection->SetSequence(ProceduralDaySequence);
+	}
+
+	UMovieScene* MovieScene = ProceduralDaySequence->GetMovieScene();
+
+	// Find the main binding
+	FGuid RootGuid = ProceduralDaySequence->FindBindingFromObject(TargetActor, TargetActor);
+	if (!RootGuid.IsValid())
+	{
+		FString RootName = TargetActor->GetName();
+		FMovieScenePossessable Possessable(RootName, TargetActor->GetClass());
+		FMovieSceneBinding     Binding(Possessable.GetGuid(), RootName);
+
+		RootGuid = Possessable.GetGuid();
+
+		// Explicitly invoke MarkAsChanged to ensure proper notification at runtime.
+		// The Modify that AddPossessable invokes only works in editor.
+		MovieScene->MarkAsChanged();
+		MovieScene->AddPossessable(Possessable, Binding);
+		ProceduralDaySequence->BindPossessableObject(RootGuid, *TargetActor, TargetActor);
+	}
+
+	// If we're trying to animate the actor, just return the root binding
+	if (Actor)
+	{
+		return RootGuid;
+	}
+
+	// If we're trying to animate a component within the actor, retrieve or create a child binding for that
+	FGuid ComponentGuid = ProceduralDaySequence->FindBindingFromObject(Component, TargetActor);
+	if (!ComponentGuid.IsValid() && Component)
+	{
+		FString Name = Component->GetName();
+		FMovieScenePossessable Possessable(Name, Component->GetClass());
+		FMovieSceneBinding     Binding(Possessable.GetGuid(), Name);
+
+		Possessable.SetParent(RootGuid, MovieScene);
+		ComponentGuid = Possessable.GetGuid();
+
+		// Explicitly invoke MarkAsChanged to ensure proper notification at runtime.
+		// The Modify that AddPossessable invokes only works in editor.
+		MovieScene->MarkAsChanged();
+		MovieScene->AddPossessable(Possessable, Binding);
+		ProceduralDaySequence->BindPossessableObject(ComponentGuid, *Component, TargetActor);
+	}
+
+	return ComponentGuid;
+}
+
+void UDaySequenceModifierComponent::AddScalarOverride(UObject* Object, FName PropertyName, double Value)
+{
+	using namespace UE::DaySequence;
+	using namespace	UE::MovieScene;
+
+	FGuid ObjectGuid = GetOrCreateProceduralBinding(Object);
+	if (!ObjectGuid.IsValid())
+	{
+		return;
+	}
+
+	check(ProceduralDaySequence);
+
+	// Create the new track
+	UMovieScene* MovieScene = ProceduralDaySequence->GetMovieScene();
+
+	FTrackInstancePropertyBindings Bindings(PropertyName, PropertyName.ToString());
+	FProperty* Property = Bindings.GetProperty(*Object);
+	if (!TestValidProperty(Object, Property))
+	{
+		// Do nothing
+	}
+	else if (Property->IsA<FFloatProperty>())
+	{
+		UMovieSceneFloatSection* FloatSection = CreateOrAddPropertyOverrideSection<UMovieSceneFloatTrack, UMovieSceneFloatSection>(MovieScene, ObjectGuid, PropertyName);
+		FloatSection->GetChannel().SetDefault(static_cast<float>(Value));
+	}
+	else if (Property->IsA<FDoubleProperty>())
+	{
+		UMovieSceneDoubleSection* DoubleSection = CreateOrAddPropertyOverrideSection<UMovieSceneDoubleTrack, UMovieSceneDoubleSection>(MovieScene, ObjectGuid, PropertyName);
+		DoubleSection->GetChannel().SetDefault(Value);
+	}
+	else
+	{
+		FFrame::KismetExecutionMessage(*FString::Printf(TEXT("Unable to animate a %s property as a scalar."), *Property->GetClass()->GetName()), ELogVerbosity::Error);
+	}
+}
+
+void UDaySequenceModifierComponent::AddColorOverride(UObject* Object, FName PropertyName, FLinearColor Value)
+{
+	using namespace UE::DaySequence;
+
+	FGuid ObjectGuid = GetOrCreateProceduralBinding(Object);
+	if (!ObjectGuid.IsValid())
+	{
+		return;
+	}
+
+	check(ProceduralDaySequence);
+
+	// Create the new track
+	UMovieScene* MovieScene = ProceduralDaySequence->GetMovieScene();
+
+	FTrackInstancePropertyBindings Bindings(PropertyName, PropertyName.ToString());
+	FProperty* Property = Bindings.GetProperty(*Object);
+	if (!TestValidProperty(Object, Property))
+	{
+		// Do nothing
+	}
+	else if (Property->IsA<FStructProperty>() && 
+		(CastField<FStructProperty>(Property)->Struct == TBaseStructure<FLinearColor>::Get() || CastField<FStructProperty>(Property)->Struct == TBaseStructure<FColor>::Get()) )
+	{
+		UMovieSceneColorSection* ColorSection = CreateOrAddPropertyOverrideSection<UMovieSceneColorTrack, UMovieSceneColorSection>(MovieScene, ObjectGuid, PropertyName);
+
+		ColorSection->GetRedChannel().SetDefault(Value.R);
+		ColorSection->GetGreenChannel().SetDefault(Value.G);
+		ColorSection->GetBlueChannel().SetDefault(Value.B);
+		ColorSection->GetAlphaChannel().SetDefault(Value.A);
+	}
+}
+
+void UDaySequenceModifierComponent::AddMaterialOverride(UObject* Object, int32 MaterialIndex, UMaterialInterface* Value)
+{
+	using namespace UE::DaySequence;
+
+	FGuid ObjectGuid = GetOrCreateProceduralBinding(Object);
+	if (!ObjectGuid.IsValid())
+	{
+		return;
+	}
+
+	check(ProceduralDaySequence);
+
+	// Create the new track
+	UMovieScene* MovieScene = ProceduralDaySequence->GetMovieScene();
+
+	UMovieScenePrimitiveMaterialTrack* MaterialTrack = CreateOrAddOverrideTrack<UMovieScenePrimitiveMaterialTrack>(MovieScene, ObjectGuid);
+	MaterialTrack->SetMaterialInfo(FComponentMaterialInfo{ FName(), MaterialIndex, EComponentMaterialType::IndexedMaterial });
+
+	UMovieScenePrimitiveMaterialSection* Section = Cast<UMovieScenePrimitiveMaterialSection>(MaterialTrack->GetAllSections()[0]);
+	Section->MaterialChannel.SetDefault(Value);
+}
+
+void UDaySequenceModifierComponent::AddScalarMaterialParameterOverride(UObject* Object, int32 MaterialIndex, FName ParameterName, float Value)
+{
+	using namespace UE::DaySequence;
+
+	FGuid ObjectGuid = GetOrCreateProceduralBinding(Object);
+	if (!ObjectGuid.IsValid())
+	{
+		return;
+	}
+
+	check(ProceduralDaySequence);
+
+	// Create the new track or locate an existing one
+	UMovieScene* MovieScene = ProceduralDaySequence->GetMovieScene();
+
+	// Material parameter tracks use the material index as the unique name
+	FName IndexAsName(*FString::FromInt(MaterialIndex));
+	UMovieSceneComponentMaterialTrack* MaterialTrack = CreateOrAddOverrideTrack<UMovieSceneComponentMaterialTrack>(MovieScene, ObjectGuid, IndexAsName);
+
+	MaterialTrack->SetMaterialInfo(FComponentMaterialInfo{FName(), MaterialIndex, EComponentMaterialType::IndexedMaterial });
+	MaterialTrack->AddScalarParameterKey(ParameterName, 0, Value);
+}
+
+void UDaySequenceModifierComponent::AddColorMaterialParameterOverride(UObject* Object, int32 MaterialIndex, FName ParameterName, FLinearColor Value)
+{
+	using namespace UE::DaySequence;
+
+	FGuid ObjectGuid = GetOrCreateProceduralBinding(Object);
+	if (!ObjectGuid.IsValid())
+	{
+		return;
+	}
+
+	check(ProceduralDaySequence);
+
+	// Create the new track or locate an existing one
+	UMovieScene* MovieScene = ProceduralDaySequence->GetMovieScene();
+
+	// Material parameter tracks use the material index as the unique name
+	FName IndexAsName(*FString::FromInt(MaterialIndex));
+	UMovieSceneComponentMaterialTrack* MaterialTrack = CreateOrAddOverrideTrack<UMovieSceneComponentMaterialTrack>(MovieScene, ObjectGuid, IndexAsName);
+
+	MaterialTrack->SetMaterialInfo(FComponentMaterialInfo{ FName(), MaterialIndex, EComponentMaterialType::IndexedMaterial });
+	MaterialTrack->AddColorParameterKey(ParameterName, 0, Value);
+}
+
+void UDaySequenceModifierComponent::AddVectorOverride(UObject* Object, FName PropertyName, FVector Value)
+{
+	using namespace UE::DaySequence;
+
+	FGuid ObjectGuid = GetOrCreateProceduralBinding(Object);
+	if (!ObjectGuid.IsValid())
+	{
+		return;
+	}
+
+	check(ProceduralDaySequence);
+
+	// Create the new track
+	UMovieScene* MovieScene = ProceduralDaySequence->GetMovieScene();
+
+	FTrackInstancePropertyBindings Bindings(PropertyName, PropertyName.ToString());
+	FProperty* Property = Bindings.GetProperty(*Object);
+	if (!TestValidProperty(Object, Property))
+	{
+		// Do nothing
+	}
+	else if (Property->IsA<FStructProperty>() && CastField<FStructProperty>(Property)->Struct == TBaseStructure<FVector>::Get())
+	{
+		UMovieSceneDoubleVectorSection* VectorSection = CreateOrAddPropertyOverrideSection<UMovieSceneDoubleVectorTrack, UMovieSceneDoubleVectorSection>(MovieScene, ObjectGuid, PropertyName);
+
+		VectorSection->SetChannelsUsed(3);
+
+		VectorSection->GetChannelProxy().GetChannel<FMovieSceneDoubleChannel>(0)->SetDefault(Value.X);
+		VectorSection->GetChannelProxy().GetChannel<FMovieSceneDoubleChannel>(1)->SetDefault(Value.Y);
+		VectorSection->GetChannelProxy().GetChannel<FMovieSceneDoubleChannel>(2)->SetDefault(Value.Z);
+	}
+}
+
+void UDaySequenceModifierComponent::AddTransformOverride(UObject* Object, FTransform Value)
+{
+	using namespace UE::DaySequence;
+
+	FGuid ObjectGuid = GetOrCreateProceduralBinding(Object);
+	if (!ObjectGuid.IsValid())
+	{
+		return;
+	}
+
+	check(ProceduralDaySequence);
+
+	// Create the new track
+	UMovieScene* MovieScene = ProceduralDaySequence->GetMovieScene();
+
+	UMovieScene3DTransformSection* TransformSection = CreateOrAddPropertyOverrideSection<UMovieScene3DTransformTrack, UMovieScene3DTransformSection>(MovieScene, ObjectGuid, "Transform");
+
+	TransformSection->GetChannelProxy().GetChannel<FMovieSceneDoubleChannel>(0)->SetDefault(Value.GetLocation().X);
+	TransformSection->GetChannelProxy().GetChannel<FMovieSceneDoubleChannel>(1)->SetDefault(Value.GetLocation().Y);
+	TransformSection->GetChannelProxy().GetChannel<FMovieSceneDoubleChannel>(2)->SetDefault(Value.GetLocation().Z);
+
+	TransformSection->GetChannelProxy().GetChannel<FMovieSceneDoubleChannel>(3)->SetDefault(Value.Rotator().Roll);
+	TransformSection->GetChannelProxy().GetChannel<FMovieSceneDoubleChannel>(4)->SetDefault(Value.Rotator().Pitch);
+	TransformSection->GetChannelProxy().GetChannel<FMovieSceneDoubleChannel>(5)->SetDefault(Value.Rotator().Yaw);
+
+	TransformSection->GetChannelProxy().GetChannel<FMovieSceneDoubleChannel>(6)->SetDefault(Value.GetScale3D().X);
+	TransformSection->GetChannelProxy().GetChannel<FMovieSceneDoubleChannel>(7)->SetDefault(Value.GetScale3D().Y);
+	TransformSection->GetChannelProxy().GetChannel<FMovieSceneDoubleChannel>(8)->SetDefault(Value.GetScale3D().Z);
+}
+
+void UDaySequenceModifierComponent::AddStaticTimeOfDayOverride(ADaySequenceActor* Actor, float Hours)
+{
+	using namespace UE::DaySequence;
+
+	FGuid ObjectGuid = GetOrCreateProceduralBinding(Actor);
+	if (!ObjectGuid.IsValid())
+	{
+		return;
+	}
+
+	check(ProceduralDaySequence);
+
+	// Create the new track
+	UMovieScene* MovieScene = ProceduralDaySequence->GetMovieScene();
+
+	UMovieSceneFloatSection* Section = CreateOrAddPropertyOverrideSection<UMovieSceneFloatTrack, UMovieSceneFloatSection>(MovieScene, ObjectGuid, "StaticTimeOfDay");
+	Section->GetChannel().SetDefault(Hours);
+}
+
+void UDaySequenceModifierComponent::AddBoolOverride(UObject* Object, FName PropertyName, bool bValue)
+{
+	using namespace UE::DaySequence;
+
+	FGuid ObjectGuid = GetOrCreateProceduralBinding(Object);
+	if (!ObjectGuid.IsValid())
+	{
+		return;
+	}
+
+	check(ProceduralDaySequence);
+
+	// Create the new track
+	UMovieScene* MovieScene = ProceduralDaySequence->GetMovieScene();
+
+	FTrackInstancePropertyBindings Bindings(PropertyName, PropertyName.ToString());
+	FProperty* Property = Bindings.GetProperty(*Object);
+	if (!TestValidProperty(Object, Property))
+	{
+		// Do nothing
+	}
+	else if (Property->IsA<FBoolProperty>())
+	{
+		UMovieSceneBoolSection* Section = CreateOrAddPropertyOverrideSection<UMovieSceneBoolTrack, UMovieSceneBoolSection>(MovieScene, ObjectGuid, PropertyName);
+		Section->GetChannel().SetDefault(bValue);
+	}
+}
+
+void UDaySequenceModifierComponent::AddVisibilityOverride(UObject* Object, bool bValue)
+{
+	using namespace UE::DaySequence;
+
+	FGuid ObjectGuid = GetOrCreateProceduralBinding(Object);
+	if (!ObjectGuid.IsValid())
+	{
+		return;
+	}
+
+	check(ProceduralDaySequence);
+
+	// Create the new track
+	UMovieScene* MovieScene = ProceduralDaySequence->GetMovieScene();
+
+	static const FName ActorVisibilityTrackName = TEXT("bHidden");
+	static const FName ComponentVisibilityTrackName = TEXT("bHiddenInGame");
+
+	const bool bIsComponent = Object->IsA<USceneComponent>();
+	const bool bIsActor     = Object->IsA<AActor>();
+
+	if (bIsComponent)
+	{
+		UMovieSceneBoolSection* VisibilitySection = CreateOrAddPropertyOverrideSection<UMovieSceneVisibilityTrack, UMovieSceneBoolSection>(MovieScene, ObjectGuid, "bHiddenInGame");
+		VisibilitySection->GetChannel().SetDefault(bValue);
+	}
+	else if (bIsActor)
+	{
+		UMovieSceneBoolSection* VisibilitySection = CreateOrAddPropertyOverrideSection<UMovieSceneVisibilityTrack, UMovieSceneBoolSection>(MovieScene, ObjectGuid, "bHidden");
+		VisibilitySection->GetChannel().SetDefault(bValue);
+	}
+}
+
+void UDaySequenceModifierComponent::SetUserDaySequence(UDaySequence* InDaySequence)
+{
+	UserDaySequence = InDaySequence;
+	ReinitializeSubSequence(nullptr);
+}
+
+float UDaySequenceModifierComponent::GetDistanceBlendFactor(FVector Position) const
+{
+	auto GetDistanceBlendFactorForShape = [this, Position](const UShapeComponent* Shape)
+	{
+		const float Distance = UE::DaySequence::ComputeSignedDistance(Shape, Position);
+		return Distance < 0.f ? FMath::Clamp(-Distance / BlendAmount, 0.f, 1.f) : 0.f;
+	};
+
+	float DistanceBlendFactor = 0.f;
+
+	for (const UShapeComponent* Shape : GetVolumeShapeComponents())
+	{
+		DistanceBlendFactor = FMath::Max(DistanceBlendFactor, GetDistanceBlendFactorForShape(Shape));
+	}
+
+	return DistanceBlendFactor;
+}
+
+TArray<UShapeComponent*> UDaySequenceModifierComponent::GetVolumeShapeComponents() const
+{
+	TArray<UShapeComponent*> ResolvedVolumeShapeComponents;
+	ResolvedVolumeShapeComponents.Reserve(VolumeShapeComponents.Num());
+	
+	if (bCachedExternalShapesInvalid)
+	{
+		UpdateCachedExternalShapes();
+	}
+	
+#if WITH_EDITOR
+	// We don't expect changes to VolumeShapeComponents or CachedExternalShapes during play.
+	// This will ensure that CachedExternalShapes remains updated to reflect editor workflows
+	// that might invalidate entries, such as deleting a shape component on an external actor.
+	bool bRecache = false;
+#endif
+
+	// This loop serves two purposes:
+	// 1) Go ahead and move from weak object ptr to raw pointer so the caller doesn't have to.
+	// 2) Determine if CachedExternalShapes is invalid so we can recache (occurs when reference shape component is deleted).
+	for (TWeakObjectPtr<UShapeComponent> Shape : CachedExternalShapes)
+	{
+		if (UShapeComponent* ValidShape = Shape.Get())
+		{
+			ResolvedVolumeShapeComponents.Add(ValidShape);
+		}
+#if WITH_EDITOR
+		else
+		{
+			// Break out out here as we will update the cached shapes and reconstruct ResolvedVolumeShapeComponents below.
+			bRecache = true;
+			break;
+		}
+#endif
+	}
+
+#if WITH_EDITOR
+	// We do this here so that we don't modify CachedExternalShapes while iterating over it.
+	// The idea is if we recache immediately before the recursive call, we should not be able to recursively hit this branch.
+	if (bRecache)
+	{
+		checkNoRecursion();
+		bCachedExternalShapesInvalid = true;
+		UpdateCachedExternalShapes();
+		return GetVolumeShapeComponents();
+	}
+#endif
+	
+	return ResolvedVolumeShapeComponents;
+}
+
+float UDaySequenceModifierComponent::GetCurrentBlendWeight() const
+{
+	FVector Position;
+
+#if WITH_EDITOR
+	if (!GetWorld()->IsGameWorld())
+	{
+		Position = UE::DaySequence::GVolumePreviewLocation;
+	}
+	else
+#endif
+	if (const AActor* BlendTarget = ExternalVolumeBlendTarget.Get())
+	{
+		Position = BlendTarget->GetActorLocation();
+	}
+	else
+	{
+		return FMath::Min(1.f, CustomVolumeBlendWeight);
+	}
+
+	return FMath::Min(GetDistanceBlendFactor(Position), CustomVolumeBlendWeight);
+}
+
+void UDaySequenceModifierComponent::SetVolumeCollisionEnabled(const ECollisionEnabled::Type InCollisionType) const
+{
+	for (UShapeComponent* Shape : GetVolumeShapeComponents())
+	{
+		Shape->SetCollisionEnabled(InCollisionType);
+	}
+}
+
+void UDaySequenceModifierComponent::EmptyVolumeShapeComponents()
+{
+	VolumeShapeComponents.Empty();
+	bCachedExternalShapesInvalid = true;
+}
+
+void UDaySequenceModifierComponent::AddVolumeShapeComponent(const FComponentReference& InShapeReference)
+{
+	VolumeShapeComponents.Add(InShapeReference);
+	bCachedExternalShapesInvalid = true;
+}
+
+void UDaySequenceModifierComponent::InvalidateMuteStates() const
+{
+	OnInvalidateMuteStates.Broadcast();
+}
+
+void UDaySequenceModifierComponent::EnableDistanceVolumeBlends(AActor* InActor)
+{
+	ExternalVolumeBlendTarget = InActor;
+}
+
+void UDaySequenceModifierComponent::SetUseVolume(bool bState)
+{
+	bUseVolume = bState;
+}
+
+void UDaySequenceModifierComponent::SetCustomVolumeBlendWeight(float Weight)
+{
+	CustomVolumeBlendWeight = FMath::Clamp(Weight, 0.f, 1.f);
+}
+
+#if ENABLE_DRAW_DEBUG
+void UDaySequenceModifierComponent::OnDebugLevelChanged(int32 InDebugLevel)
+{
+	DebugLevel = InDebugLevel;
+}
+
+bool UDaySequenceModifierComponent::ShouldShowDebugInfo() const
+{
+	if (GetOwner()->HasAuthority())
+	{
+		return false;
+	}
+	
+	switch (DebugLevel)
+	{
+	case 0: return false;
+	case 1: return bIsEnabled;
+	case 2: return bIsComponentEnabled;
+	case 3: return true;
+	default: return false;
+	}
+}
+#endif
+
+bool UDaySequenceModifierComponent::IsBlendTargetInAnyVolume()
+{
+	OccupiedVolumes = 0;
+
+	for (const UShapeComponent* Shape : GetVolumeShapeComponents())
+	{
+		if (ExternalVolumeBlendTarget.IsValid() && Shape->IsOverlappingActor(ExternalVolumeBlendTarget.Get()))
+		{
+			OccupiedVolumes++;
+		}
+	}
+
+	return OccupiedVolumes > 0;
+}
+
+void UDaySequenceModifierComponent::UpdateCachedExternalShapes() const
+{
+	check (bCachedExternalShapesInvalid)
+	
+	CachedExternalShapes.Reset();
+	
+	for (const FComponentReference& ComponentRef : VolumeShapeComponents)
+	{
+		if (ComponentRef.PathToComponent.Len() != 0 || ComponentRef.ComponentProperty != NAME_None || !ComponentRef.OverrideComponent.IsExplicitlyNull())
+		{
+			if (UShapeComponent* ResolvedShape = Cast<UShapeComponent>(ComponentRef.GetComponent(GetOwner())); IsValid(ResolvedShape))
+			{
+				CachedExternalShapes.Add(ResolvedShape);
+			}
+		}
+	}
+
+	bCachedExternalShapesInvalid = false;
+	
+	UnbindOverlapEvents();
+	BindOverlapEvents();
+}
+
+void UDaySequenceModifierComponent::BindOverlapEvents() const
+{
+	for (UShapeComponent* Shape : GetVolumeShapeComponents())
+	{
+		// AddUniqueDynamic prevents accidental duplicate binding
+        Shape->OnComponentBeginOverlap.AddUniqueDynamic(this, &UDaySequenceModifierComponent::OnVolumeOverlapBegin);
+        Shape->OnComponentEndOverlap.AddUniqueDynamic(this, &UDaySequenceModifierComponent::OnVolumeOverlapEnd);
+	}
+}
+
+void UDaySequenceModifierComponent::UnbindOverlapEvents() const
+{
+	for (UShapeComponent* Shape : GetVolumeShapeComponents())
+	{
+		Shape->OnComponentBeginOverlap.RemoveDynamic(this, &UDaySequenceModifierComponent::OnVolumeOverlapBegin);
+		Shape->OnComponentEndOverlap.RemoveDynamic(this, &UDaySequenceModifierComponent::OnVolumeOverlapEnd);
+	}
+}
+
+void UDaySequenceModifierComponent::OnVolumeOverlapBegin(UPrimitiveComponent* OverlappedComponent, AActor* OtherActor, UPrimitiveComponent* OtherComp, int32 OtherBodyIndex, bool bFromSweep, const FHitResult& SweepResult)
+{
+	if (ExternalVolumeBlendTarget.IsValid() && OtherActor == ExternalVolumeBlendTarget)
+	{
+		OccupiedVolumes++;
+
+		if (OccupiedVolumes == 1 && GetOwner() && !GetOwner()->HasAuthority())
+		{
+			OnVolumeBlendTargetOverlapBegin.Broadcast(ExternalVolumeBlendTarget.Get());
+			OnVolumeBlendTargetOverlapBeginNative.Broadcast(ExternalVolumeBlendTarget.Get());
+		}
+	}
+}
+
+void UDaySequenceModifierComponent::OnVolumeOverlapEnd(UPrimitiveComponent* OverlappedComponent, AActor* OtherActor, UPrimitiveComponent* OtherComp, int32 OtherBodyIndex)
+{
+	if (ExternalVolumeBlendTarget.IsValid() && OtherActor == ExternalVolumeBlendTarget)
+	{
+		OccupiedVolumes--;
+
+		if (OccupiedVolumes == 0 && GetOwner() && !GetOwner()->HasAuthority())
+		{
+			OnVolumeBlendTargetOverlapEnd.Broadcast(ExternalVolumeBlendTarget.Get());
+			OnVolumeBlendTargetOverlapEndNative.Broadcast(ExternalVolumeBlendTarget.Get());
+		}
+	}
+}
+
+
+#undef LOCTEXT_NAMESPACE
diff --git a/Engine/Plugins/Experimental/DaySequence/Source/DaySequence/Private/DaySequenceProvider.cpp b/Engine/Plugins/Experimental/DaySequence/Source/DaySequence/Private/DaySequenceProvider.cpp
new file mode 100644
--- /dev/null
+++ b/Engine/Plugins/Experimental/DaySequence/Source/DaySequence/Private/DaySequenceProvider.cpp
@@ -0,0 +1,46 @@
+// Copyright Epic Games, Inc. All Rights Reserved.
+
+#include "DaySequenceProvider.h"
+#include "DaySequenceActor.h"
+#include "DaySequenceSubsystem.h"
+
+#include "Engine/World.h"
+
+#include UE_INLINE_GENERATED_CPP_BY_NAME(DaySequenceProvider)
+
+ADaySequenceProvider::ADaySequenceProvider(const FObjectInitializer& Init)
+: Super(Init)
+{
+	USceneComponent* SceneRootComponent = CreateDefaultSubobject<USceneComponent>(USceneComponent::GetDefaultSceneRootVariableName());
+	SetRootComponent(SceneRootComponent);
+}
+
+TArrayView<TObjectPtr<UDaySequence>> ADaySequenceProvider::GetDaySequences()
+{
+	return DaySequenceAssets;
+}
+
+#if WITH_EDITOR
+void ADaySequenceProvider::PostEditChangeProperty(FPropertyChangedEvent& PropertyChangedEvent)
+{
+	const FName PropertyName = (PropertyChangedEvent.Property != nullptr) ? PropertyChangedEvent.Property->GetFName() : NAME_None;
+	if (PropertyName == GET_MEMBER_NAME_CHECKED(ADaySequenceProvider, DaySequenceAssets))
+	{
+		// Force update the DaySequenceActor's sequence.
+		// TODO: only update if this actor is assigned to the DayActor.
+		if (const UDaySequenceSubsystem* DaySubsystem = GetWorld()->GetSubsystem<UDaySequenceSubsystem>())
+		{
+			if (ADaySequenceActor* DayActor = DaySubsystem->GetDaySequenceActor())
+			{
+				DayActor->UpdateRootSequence();
+			}
+		}
+	}
+
+	Super::PostEditChangeProperty(PropertyChangedEvent);
+}
+#endif
+
+
+
+
diff --git a/Engine/Plugins/Experimental/DaySequence/Source/DaySequence/Private/EnvironmentLightingActor.cpp b/Engine/Plugins/Experimental/DaySequence/Source/DaySequence/Private/EnvironmentLightingActor.cpp
new file mode 100644
--- /dev/null
+++ b/Engine/Plugins/Experimental/DaySequence/Source/DaySequence/Private/EnvironmentLightingActor.cpp
@@ -0,0 +1,59 @@
+// Copyright Epic Games, Inc. All Rights Reserved.
+
+#include "EnvironmentLightingActor.h"
+
+#include "Components/DirectionalLightComponent.h"
+#include "Components/ExponentialHeightFogComponent.h"
+#include "Components/SkyAtmosphereComponent.h"
+#include "Components/SkyLightComponent.h"
+#include "Components/VolumetricCloudComponent.h"
+
+#include UE_INLINE_GENERATED_CPP_BY_NAME(EnvironmentLightingActor)
+
+FName AEnvironmentLightingActor::SkyAtmosphereComponentName(TEXT("SkyAtmosphere"));
+FName AEnvironmentLightingActor::SkyLightComponentName(TEXT("SkyLight"));
+FName AEnvironmentLightingActor::SunRootComponentName(TEXT("SunRoot"));
+FName AEnvironmentLightingActor::SunComponentName(TEXT("Sun"));
+FName AEnvironmentLightingActor::ExponentialHeightFogComponentName(TEXT("ExponentialHeightFog"));
+FName AEnvironmentLightingActor::VolumetricCloudComponentName(TEXT("VolumetricClouds"));
+
+AEnvironmentLightingActor::AEnvironmentLightingActor(const FObjectInitializer& Init)
+: Super(Init)
+{
+	SkyAtmosphereComponent = CreateOptionalDefaultSubobject<USkyAtmosphereComponent>(SkyAtmosphereComponentName);
+	if (SkyAtmosphereComponent)
+	{
+		SkyAtmosphereComponent->SetupAttachment(RootComponent);
+	}
+
+	SkyLightComponent = CreateOptionalDefaultSubobject<USkyLightComponent>(SkyLightComponentName);
+	if (SkyLightComponent)
+	{
+		SkyLightComponent->SetupAttachment(RootComponent);
+	}
+
+	SunRootComponent = CreateOptionalDefaultSubobject<USceneComponent>(SunRootComponentName);
+	if (SunRootComponent)
+	{
+		SunRootComponent->SetupAttachment(RootComponent);
+	}
+
+	SunComponent = CreateOptionalDefaultSubobject<UDirectionalLightComponent>(SunComponentName);
+	if (SunComponent)
+	{
+		SunComponent->SetupAttachment(SunRootComponent ? SunRootComponent : RootComponent);
+	}
+
+	ExponentialHeightFogComponent = CreateOptionalDefaultSubobject<UExponentialHeightFogComponent>(ExponentialHeightFogComponentName);
+	if (ExponentialHeightFogComponent)
+	{
+		ExponentialHeightFogComponent->SetupAttachment(RootComponent);
+	}
+
+	VolumetricCloudComponent = CreateOptionalDefaultSubobject<UVolumetricCloudComponent>(VolumetricCloudComponentName);
+	if (VolumetricCloudComponent)
+	{
+		VolumetricCloudComponent->SetupAttachment(RootComponent);
+	}
+}
+
diff --git a/Engine/Plugins/Experimental/DaySequence/Source/DaySequence/Public/Actors/DaySequenceModifierVolume.h b/Engine/Plugins/Experimental/DaySequence/Source/DaySequence/Public/Actors/DaySequenceModifierVolume.h
new file mode 100644
--- /dev/null
+++ b/Engine/Plugins/Experimental/DaySequence/Source/DaySequence/Public/Actors/DaySequenceModifierVolume.h
@@ -0,0 +1,125 @@
+// Copyright Epic Games, Inc. All Rights Reserved.
+
+#pragma once
+
+#include "CoreMinimal.h"
+#include "GameFramework/Actor.h"
+
+#include "DaySequenceModifierVolume.generated.h"
+
+class ADaySequenceActor;
+class APawn;
+
+class UBoxComponent;
+class UDaySequenceModifierComponent;
+
+UCLASS(Blueprintable)
+class DAYSEQUENCE_API ADaySequenceModifierVolume : public AActor
+{
+	GENERATED_BODY()
+
+public:
+	ADaySequenceModifierVolume(const FObjectInitializer& Init);
+
+	UFUNCTION(BlueprintImplementableEvent)
+	void OnDaySequenceActorBound(ADaySequenceActor* InActor);
+	
+protected:
+	//~ Begin AActor interface
+	virtual void BeginPlay() override;
+	virtual void OnConstruction(const FTransform& Transform) override;
+	//~ End AActor interface
+	
+	// Primary initialization entry point to clarify initialization control flow and to consolidate editor and runtime initialization.
+	// Called from BeginPlay or OnConstruction
+	void Initialize();
+
+	
+	
+	/** Volume Setup Functions */
+	
+	/**
+	 * Unconditionally adds shape components to the day sequence modifier via AddShapeComponentsToModifier.
+	 * At runtime this also unconditionally calls SetupVolumeCallbacks.
+	 * In editor this conditionally calls SetupVolumeCallbacks if we are not in an Editor world.
+	 */
+	void VolumeSetup();
+
+	/**
+	 * By default this adds an FComponentReference pointing to DefaultBox to the day sequence modifier.
+	 * Derived classes should override this function to add their own shape components at the proper time.
+	 * If overridden, calling Super::AddShapeComponentsToModifier will determine if DefaultBox is used or ignored.
+	 */
+	virtual void AddShapeComponentsToModifier();
+
+	/** Binds the day sequence modifier overlap events to OnOverlapBegin and OnOverlapEnd. */
+	void SetupVolumeCallbacks();
+	void OnOverlapBegin() const;
+	void OnOverlapEnd() const;
+
+	
+	
+	/** Player Controller Setup Functions */
+	
+	/**
+	 * At runtime this unconditionally calls CachePlayerController and SetupBlendTargetCallbacks.
+	 * In editor this conditionally calls CachePlayerController and SetupBlendTargetCallbacks if we are not in an Editor world.
+	 */
+	void PlayerControllerSetup();
+
+	/**
+	 * Attempts to find a local player controller.
+	 * If found, we get the controller's view target.
+	 * If not found, we call QueuePlayerControllerQuery which will result in another call to this function next tick.
+	 */
+	void CachePlayerController();
+
+	/** Set a timer to call CachePlayerController one tick from now. */
+	void QueuePlayerControllerQuery();
+
+	/** Register a callback to handle changes to our player controller's view target. */
+	void SetupBlendTargetCallbacks();
+	
+
+	
+	/** Day Sequence Actor Setup Functions */
+
+	/** Dumb wrapper for SetupDaySequenceSubsystemCallbacks and BindToDaySequenceActor. */
+	void DaySequenceActorSetup();
+
+	/** Registers a callback that calls BindToDaySequenceActor when the world's current Day Sequence Actor changes. */
+	void SetupDaySequenceSubsystemCallbacks();
+
+	/**
+	 * Attempts to get the world's current Day Sequence Actor.
+	 * 
+	 * If a Day Sequence Actor is found that is not the currently bound actor,
+	 * we notify the day sequence modifier and call VolumeSetup and PlayerControllerSetup.
+	 */
+	void BindToDaySequenceActor();
+	
+
+	
+	/** Enables this modifier if the current blend target is in the volume. */
+	void TryEnableModifier() const;
+	
+	void SetBlendTarget(AActor* InPawn);
+	
+protected:
+	UPROPERTY(VisibleAnywhere, Category = "Day Sequence", BlueprintReadOnly, meta = (AllowPrivateAccess = "true"))
+	TObjectPtr<UDaySequenceModifierComponent> DaySequenceModifier;
+
+	UPROPERTY(VisibleAnywhere, Category = "Day Sequence", BlueprintReadOnly, meta = (AllowPrivateAccess = "true"))
+	TObjectPtr<UBoxComponent> DefaultBox;
+	
+	UPROPERTY(Transient, DuplicateTransient)
+    TObjectPtr<ADaySequenceActor> DaySequenceActor;
+	
+	UPROPERTY(Transient, DuplicateTransient)
+	TObjectPtr<APlayerController> CachedPlayerController;
+	
+	UPROPERTY(Transient, DuplicateTransient)
+	TObjectPtr<AActor> CurrentBlendTarget;
+	
+	FDelegateHandle ViewTargetChangedHandle;
+};
diff --git a/Engine/Plugins/Experimental/DaySequence/Source/DaySequence/Public/Actors/ProceduralDaySequenceActor.h b/Engine/Plugins/Experimental/DaySequence/Source/DaySequence/Public/Actors/ProceduralDaySequenceActor.h
new file mode 100644
--- /dev/null
+++ b/Engine/Plugins/Experimental/DaySequence/Source/DaySequence/Public/Actors/ProceduralDaySequenceActor.h
@@ -0,0 +1,44 @@
+﻿// Copyright Epic Games, Inc. All Rights Reserved.
+
+#pragma once
+
+#include "CoreMinimal.h"
+#include "SimpleDaySequenceActor.h"
+
+#include "ProceduralDaySequenceActor.generated.h"
+
+class UProceduralDaySequenceBuilder;
+
+UCLASS(Abstract)
+class DAYSEQUENCE_API AProceduralDaySequenceActor : public ASimpleDaySequenceActor
+{
+	GENERATED_BODY()
+
+protected:
+	virtual void OnConstruction(const FTransform& Transform) override;
+
+#if WITH_EDITOR
+	virtual void PostEditChangeProperty(FPropertyChangedEvent& PropertyChangedEvent) override;
+#endif
+	
+	/**
+	 * Populates ProceduralSequence with keys.
+	 * By default this does nothing, derived classes should override this to modify ProceduralSequence through the provided builder.
+	 * 
+	 * @param SequenceBuilder A ready-to-use sequence builder.
+	 */
+	virtual void BuildSequence(UProceduralDaySequenceBuilder* SequenceBuilder);
+
+	virtual void InitializeDaySequences() override final;
+
+protected:
+	bool bProceduralSequenceInvalid = false;
+	
+private:
+	// Wrapper that instantiates the sequence builder, invokes BuildSequence, and assigns the created sequence to ProceduralSequence.
+	void CreateProceduralSequence();
+	
+	/** The procedural sequence. This is constructed in BuildSequence and added to the RootSequence in InitializeDaySequences. */
+	UPROPERTY(Transient)
+	TObjectPtr<UDaySequence> ProceduralSequence;
+};
diff --git a/Engine/Plugins/Experimental/DaySequence/Source/DaySequence/Public/Actors/SimpleDaySequenceActor.h b/Engine/Plugins/Experimental/DaySequence/Source/DaySequence/Public/Actors/SimpleDaySequenceActor.h
new file mode 100644
--- /dev/null
+++ b/Engine/Plugins/Experimental/DaySequence/Source/DaySequence/Public/Actors/SimpleDaySequenceActor.h
@@ -0,0 +1,56 @@
+﻿// Copyright Epic Games, Inc. All Rights Reserved.
+
+#pragma once
+
+#include "CoreMinimal.h"
+#include "DaySequenceActor.h"
+
+#include "SimpleDaySequenceActor.generated.h"
+
+class USkyAtmosphereComponent;
+class USkyLightComponent;
+class UVolumetricCloudComponent;
+class UDirectionalLightComponent;
+class UExponentialHeightFogComponent;
+
+/**
+ * A self-registering Day Sequence Actor with a simple set of lighting components (some of which are optional).
+ * Can be used as-is or extended by derived classes (see AProceduralDaySequenceActor for example).
+ */
+UCLASS(Blueprintable, HideCategories=(Tags, Transform, Networking, LevelInstance))
+class DAYSEQUENCE_API ASimpleDaySequenceActor
+	: public ADaySequenceActor
+{
+	GENERATED_BODY()
+
+public:
+	ASimpleDaySequenceActor(const FObjectInitializer& Init);
+
+protected:
+	/** BeginPlay and OnConstruction overrides auto-register this actor with the DaySequenceSubsystem. */
+	virtual void BeginPlay() override;
+	virtual void OnConstruction(const FTransform& Transform) override;
+
+protected:
+	/** Standard Components **/
+	
+	UPROPERTY(VisibleDefaultsOnly, BlueprintReadOnly, Category= "Day Sequence", meta = (AllowPrivateAccess = "true"))
+	TObjectPtr<USceneComponent> SunRootComponent;
+	
+	UPROPERTY(VisibleDefaultsOnly, BlueprintReadOnly, Category= "Day Sequence", meta = (AllowPrivateAccess = "true"))
+	TObjectPtr<UDirectionalLightComponent> SunComponent;
+	
+	UPROPERTY(VisibleDefaultsOnly, BlueprintReadOnly, Category = "Day Sequence", meta = (AllowPrivateAccess = "true"))
+	TObjectPtr<USkyAtmosphereComponent> SkyAtmosphereComponent;
+	
+	UPROPERTY(VisibleDefaultsOnly, BlueprintReadOnly, Category = "Day Sequence", meta = (AllowPrivateAccess = "true"))
+	TObjectPtr<USkyLightComponent> SkyLightComponent;
+	
+	/** Optional Components **/
+	
+	UPROPERTY(VisibleDefaultsOnly, BlueprintReadOnly, Category = "Day Sequence", meta = (AllowPrivateAccess = "true"))
+	TObjectPtr<UExponentialHeightFogComponent> ExponentialHeightFogComponent;
+	
+	UPROPERTY(VisibleDefaultsOnly, BlueprintReadOnly, Category = "Day Sequence", meta = (AllowPrivateAccess = "true"))
+	TObjectPtr<UVolumetricCloudComponent> VolumetricCloudComponent;
+};
diff --git a/Engine/Plugins/Experimental/DaySequence/Source/DaySequence/Public/Actors/SunAngleDaySequenceActor.h b/Engine/Plugins/Experimental/DaySequence/Source/DaySequence/Public/Actors/SunAngleDaySequenceActor.h
new file mode 100644
--- /dev/null
+++ b/Engine/Plugins/Experimental/DaySequence/Source/DaySequence/Public/Actors/SunAngleDaySequenceActor.h
@@ -0,0 +1,41 @@
+﻿// Copyright Epic Games, Inc. All Rights Reserved.
+
+#pragma once
+
+#include "CoreMinimal.h"
+#include "ProceduralDaySequenceActor.h"
+
+#include "SunAngleDaySequenceActor.generated.h"
+
+/**
+ * A simple procedural Day Sequence Actor that positions the sun according to SunPitch, SunYaw, and SunRoll.
+ * A sequence is generated on the fly that animates SunPitch linearly over the length of the day.
+ * SunYaw and SunRoll can be manually configured to set the sun's elevation and the direction it moves across the sky.
+ */
+UCLASS(Blueprintable)
+class DAYSEQUENCE_API ASunAngleDaySequenceActor
+	: public AProceduralDaySequenceActor
+{
+	GENERATED_BODY()
+
+protected:
+	virtual void Tick(float DeltaTime) override;
+	
+	void ApplySunAngle();
+	
+	virtual void BuildSequence(UProceduralDaySequenceBuilder* SequenceBuilder) override;
+	virtual void SequencePlayerUpdated(float CurrentTime, float PreviousTime) override;
+
+protected:
+	// Determines how far the sun is along its path.
+	UPROPERTY(VisibleAnywhere, Interp, Category=Sun, meta=(NoResetToDefault))
+	double SunPitch = 0.0f;
+
+	// Determines the path the sun takes through the sky (for example, determines if the sun moves east to west).
+	UPROPERTY(EditAnywhere, Interp, Category=Sun, meta=(UIMin = "-180", UIMax = "180", ClampMin="-180", ClampMax="180", SliderExponent=1))
+	double SunYaw = 0.f;
+
+	// Determines how high above the horizon the sun will be at its peak.
+	UPROPERTY(EditAnywhere, Interp, Category=Sun, meta=(UIMin = "-90", UIMax = "90", ClampMin="-180", ClampMax="180", SliderExponent=1))
+	double SunRoll = 0.f;
+};
diff --git a/Engine/Plugins/Experimental/DaySequence/Source/DaySequence/Public/Actors/SunPositionDaySequenceActor.h b/Engine/Plugins/Experimental/DaySequence/Source/DaySequence/Public/Actors/SunPositionDaySequenceActor.h
new file mode 100644
--- /dev/null
+++ b/Engine/Plugins/Experimental/DaySequence/Source/DaySequence/Public/Actors/SunPositionDaySequenceActor.h
@@ -0,0 +1,67 @@
+// Copyright Epic Games, Inc. All Rights Reserved.
+
+#pragma once
+
+#include "CoreMinimal.h"
+#include "SunAngleDaySequenceActor.h"
+
+#include "SunPositionDaySequenceActor.generated.h"
+
+/**
+ * A simple procedural Day Sequence Actor that mimics an Earth day at a user-specified location and time.
+ */
+UCLASS(Blueprintable)
+class DAYSEQUENCE_API ASunPositionDaySequenceActor
+	: public AProceduralDaySequenceActor
+{
+	GENERATED_BODY()
+
+protected:
+	virtual void Tick(float DeltaTime) override;
+
+	void ApplySunAngle();
+	
+	virtual void BuildSequence(UProceduralDaySequenceBuilder* SequenceBuilder) override;
+	virtual void SequencePlayerUpdated(float CurrentTime, float PreviousTime) override;
+
+#if WITH_EDITOR
+	virtual void PostEditChangeProperty(FPropertyChangedEvent& PropertyChangedEvent) override;
+#endif
+	
+private:
+	
+	/** User settings */
+	UPROPERTY(EditAnywhere, Category = "Procedural Parameters")
+	unsigned KeyCount;
+	
+	UPROPERTY(EditAnywhere, Category = "Procedural Parameters")
+	FDateTime Time;
+
+	UPROPERTY(EditAnywhere, Category = "Procedural Parameters")
+	double TimeZone;
+	
+	UPROPERTY(EditAnywhere, Category = "Procedural Parameters")
+	double Latitude;
+
+	UPROPERTY(EditAnywhere, Category = "Procedural Parameters")
+	double Longitude;
+	
+	UPROPERTY(EditAnywhere, Category = "Procedural Parameters")
+	bool bIsDaylightSavings;
+
+	
+	/** Read only properties */
+	UPROPERTY(VisibleAnywhere, Transient, Category=Sun, meta=(NoResetToDefault))
+	double Elevation = 0.0f;
+	
+	UPROPERTY(VisibleAnywhere, Transient, Category=Sun, meta=(NoResetToDefault))
+	double Azimuth = 0.0f;
+
+	
+	/** Animated properties */
+	UPROPERTY(Transient, Interp, Category=Sun)
+	double SunPitch = 0.0f;
+
+	UPROPERTY(Transient, Interp, Category=Sun)
+	double SunYaw = 0.f;
+};
\ No newline at end of file
diff --git a/Engine/Plugins/Experimental/DaySequence/Source/DaySequence/Public/DaySequence.h b/Engine/Plugins/Experimental/DaySequence/Source/DaySequence/Public/DaySequence.h
new file mode 100644
--- /dev/null
+++ b/Engine/Plugins/Experimental/DaySequence/Source/DaySequence/Public/DaySequence.h
@@ -0,0 +1,123 @@
+// Copyright Epic Games, Inc. All Rights Reserved.
+
+#pragma once
+
+#include "Interfaces/Interface_AssetUserData.h"
+#include "MovieSceneSequence.h"
+#include "DaySequenceBindingReference.h"
+#include "DaySequence.generated.h"
+
+class UAssetUserData;
+class UBlueprint;
+struct FMovieSceneSequenceID;
+
+UCLASS(BlueprintType, HideCategories=(Object))
+class DAYSEQUENCE_API UDaySequence
+	: public UMovieSceneSequence
+	, public IInterface_AssetUserData
+{
+	GENERATED_BODY()
+public:
+	/** Pointer to the movie scene that controls this animation. */
+	UPROPERTY()
+	TObjectPtr<UMovieScene> MovieScene;
+
+public:
+	UDaySequence(const FObjectInitializer&);
+	
+	/** Initialize this sequence. */
+	virtual void Initialize();
+	virtual void Initialize(EObjectFlags Flags);
+
+	void AddDefaultBinding(const FGuid& PossessableGuid);
+
+	// UMovieSceneSequence interface
+	virtual void BindPossessableObject(const FGuid& ObjectId, UObject& PossessedObject, UObject* Context) override;
+	virtual bool CanPossessObject(UObject& Object, UObject* InPlaybackContext) const override;
+	virtual void LocateBoundObjects(const FGuid& ObjectId, UObject* Context, TArray<UObject*, TInlineAllocator<1>>& OutObjects) const override;
+	virtual FGuid FindBindingFromObject(UObject* InObject, UObject* Context) const override;
+	virtual void GatherExpiredObjects(const FMovieSceneObjectCache& InObjectCache, TArray<FGuid>& OutInvalidIDs) const override;
+	virtual UMovieScene* GetMovieScene() const override;
+	virtual UObject* GetParentObject(UObject* Object) const override;
+	virtual void UnbindPossessableObjects(const FGuid& ObjectId) override;
+	virtual void UnbindObjects(const FGuid& ObjectId, const TArray<UObject*>& InObjects, UObject* InContext) override;
+	virtual void UnbindInvalidObjects(const FGuid& ObjectId, UObject* InContext) override;
+	virtual bool AllowsSpawnableObjects() const override;
+	virtual bool CanRebindPossessable(const FMovieScenePossessable& InPossessable) const override;
+	virtual UObject* MakeSpawnableTemplateFromInstance(UObject& InSourceObject, FName ObjectName) override;
+	virtual bool CanAnimateObject(UObject& InObject) const override;
+	virtual UObject* CreateDirectorInstance(TSharedRef<const FSharedPlaybackState> SharedPlaybackState, FMovieSceneSequenceID SequenceID) override;
+	virtual void PostLoad() override;
+	virtual void PostInitProperties() override;
+	virtual bool Rename(const TCHAR* NewName = nullptr, UObject* NewOuter = nullptr, ERenameFlags Flags = REN_None) override;
+
+	//~ Begin IInterface_AssetUserData Interface
+	virtual void AddAssetUserData(UAssetUserData* InUserData) override;
+	virtual void RemoveUserDataOfClass(TSubclassOf<UAssetUserData> InUserDataClass) override;
+	virtual UAssetUserData* GetAssetUserDataOfClass(TSubclassOf<UAssetUserData> InUserDataClass) override;
+	virtual const TArray<UAssetUserData*>* GetAssetUserDataArray() const override;
+	//~ End IInterface_AssetUserData Interface
+
+#if WITH_EDITOR
+	virtual ETrackSupport IsTrackSupported(TSubclassOf<class UMovieSceneTrack> InTrackClass) const override;
+	virtual void GetAssetRegistryTags(FAssetRegistryTagsContext Context) const override;
+
+	DECLARE_DELEGATE_RetVal_OneParam(void, FPostDuplicateEvent, UDaySequence*);
+	static FPostDuplicateEvent PostDuplicateEvent;
+#endif
+
+	virtual void PostDuplicate(bool bDuplicateForPIE) override;
+
+#if WITH_EDITOR
+
+
+public:
+
+	/**
+	 * Assign a new director blueprint to this sequence. The specified blueprint *must* be contained within this object.?	 */
+	void SetDirectorBlueprint(UBlueprint* NewDirectorBlueprint);
+
+	/**
+	 * Retrieve the currently assigned director blueprint for this sequence
+	 */
+	UBlueprint* GetDirectorBlueprint() const;
+
+	FString GetDirectorBlueprintName() const;
+
+protected:
+
+	virtual FGuid CreatePossessable(UObject* ObjectToPossess) override;
+	virtual FGuid CreateSpawnable(UObject* ObjectToSpawn) override;
+
+	FGuid FindOrAddBinding(UObject* ObjectToPossess);
+
+	/**
+	 * Invoked when this sequence's director blueprint has been recompiled
+	 */
+	void OnDirectorRecompiled(UBlueprint*);
+
+#endif // WITH_EDITOR
+
+protected:
+	/** References to bound objects. */
+	UPROPERTY()
+	FDaySequenceBindingReferences BindingReferences;
+
+#if WITH_EDITORONLY_DATA
+	/** A pointer to the director blueprint that generates this sequence's DirectorClass. */
+	UPROPERTY()
+	TObjectPtr<UBlueprint> DirectorBlueprint;
+#endif
+
+	/**
+	 * The class that is used to spawn this sequence's director instance.
+	 * Director instances are allocated on-demand one per sequence during evaluation and are used by event tracks for triggering events.
+	 */
+	UPROPERTY()
+	TObjectPtr<UClass> DirectorClass;
+
+protected:
+	/** Array of user data stored with the asset */
+	UPROPERTY(EditAnywhere, AdvancedDisplay, Instanced, Category = Animation)
+	TArray<TObjectPtr<UAssetUserData>> AssetUserData;
+};
diff --git a/Engine/Plugins/Experimental/DaySequence/Source/DaySequence/Public/DaySequenceActor.h b/Engine/Plugins/Experimental/DaySequence/Source/DaySequence/Public/DaySequenceActor.h
new file mode 100644
--- /dev/null
+++ b/Engine/Plugins/Experimental/DaySequence/Source/DaySequence/Public/DaySequenceActor.h
@@ -0,0 +1,541 @@
+// Copyright Epic Games, Inc. All Rights Reserved.
+
+#pragma once
+
+#include "GameFramework/Info.h"
+#include "DaySequenceConditionSet.h"
+#include "DrawDebugHelpers.h"	// Defines ENABLE_DRAW_DEBUG
+#include "IMovieScenePlaybackClient.h"
+#include "MovieSceneBindingOwnerInterface.h"
+
+#include "Misc/Timecode.h"
+#include "DaySequenceActor.generated.h"
+
+class FDebugDisplayInfo;
+class UDaySequencePlayer;
+class UMovieSceneBindingOverrides;
+class UMovieSceneSequencePlayer;
+namespace EEndPlayReason { enum Type : int; }
+struct FMovieSceneSequencePlaybackSettings;
+
+#if ENABLE_DRAW_DEBUG
+class AHUD;
+#endif
+
+class UCurveFloat;
+class UDaySequence;
+class UDaySequenceCollectionAsset;
+class UDaySequenceConditionTag;
+class UDaySequenceTrack;
+class UMovieSceneSubSection;
+
+struct FDaySequenceCollectionEntry;
+
+#ifndef ROOT_SEQUENCE_RECONSTRUCTION_ENABLED
+	#define ROOT_SEQUENCE_RECONSTRUCTION_ENABLED WITH_EDITOR
+#endif
+
+namespace UE::DaySequence
+{
+	struct FStaticTimeControllerOverride;
+
+#if ENABLE_DRAW_DEBUG
+	// This provides methods for determining if this debug entry should be shown and for getting a pointer to the debug data.
+	// Anything (Currently only DaySequenceModifierComponents) can submit one of these entries to a DaySequenceActor.
+	// The debug data can be printed in play with the command "showdebug DaySequence" if ShowCondition evaluates to true. 
+	struct FDaySequenceDebugEntry
+	{
+		using FShowDebugDataConditionFunction = TFunction<bool()>;
+		using FGetDebugDataFunction = TFunction<TSharedPtr<TMap<FString, FString>>()>;
+
+		FDaySequenceDebugEntry(FShowDebugDataConditionFunction InShowCondition, FGetDebugDataFunction InGetData);
+		
+		FShowDebugDataConditionFunction ShowCondition;
+		FGetDebugDataFunction GetData;
+	};
+	
+	// Alias type which stores an array of weak pointers to debug entries 
+	using FDebugEntryArray = TArray<TWeakPtr<FDaySequenceDebugEntry>>;
+
+	// Alias type which defines the signature of the draw function a registered category must be associated with
+	using FDebugCategoryDrawFunction = TFunction<void(UCanvas*, TArray<TSharedPtr<TMap<FString, FString>>>&, const FString&)>;
+	
+	// Alias type which maps a category to its array of debug entries and to a callback which handles drawing the data.
+	using FDebugEntryMap = TMap<FName, TPair<FDebugEntryArray, FDebugCategoryDrawFunction>>;
+#endif
+
+	DECLARE_MULTICAST_DELEGATE(FOnInvalidateMuteStates);
+}
+
+DECLARE_DYNAMIC_MULTICAST_DELEGATE_OneParam(FOnTimeOfDayPreviewChanged, float, Time);
+
+UCLASS(Blueprintable)
+class DAYSEQUENCE_API ADaySequenceActor
+	: public AInfo
+	, public IMovieScenePlaybackClient
+	, public IMovieSceneBindingOwnerInterface
+{
+	GENERATED_BODY()
+
+public:
+	ADaySequenceActor(const FObjectInitializer& Init);
+
+	/** Access this actor's sequence player, or None if it is invalid (not yet initialized or already destroyed) */
+	UFUNCTION(BlueprintGetter)
+	UDaySequencePlayer* GetSequencePlayer() const;
+
+	/**
+	 * Returns true if the given InDaySequence is referenced by any entry in the DaySequences map property.
+	 * 
+	 * @param InDaySequence the sequence to search for
+	 * @return true if InDaySequence exists in the DaySequences map.
+	 */
+	UFUNCTION(BlueprintCallable, Category="Sequence")
+	bool ContainsDaySequence(const UDaySequence* InDaySequence);
+
+	/** Set whether or not to replicate playback for this actor */
+	UFUNCTION(BlueprintSetter)
+	void SetReplicatePlayback(bool ReplicatePlayback);
+
+#if WITH_EDITOR
+	virtual bool GetReferencedContentObjects(TArray<UObject*>& Objects) const override;
+
+	/**
+	 * Set the TimeOfDayPreview if it is currently different from the specified time.
+	 * @param InHours preview time to set in hours.
+	 */
+	void ConditionalSetTimeOfDayPreview(float InHours);
+
+	/**
+	 * Native event invoked when the TimeOfDayPreview property is changed.
+	 * @param new preview time in hours.
+	 */
+	DECLARE_EVENT_OneParam(ADaySequenceActor, FOnTimeOfDayPreviewChangedEvent, float)
+	FOnTimeOfDayPreviewChangedEvent OnTimeOfDayPreviewChangedEvent;
+#endif //WITH_EDITOR
+
+	/** @return the preview time in hours. */
+	UFUNCTION(BlueprintCallable, Category="TimeOfDay")
+	float GetTimeOfDayPreview() const;
+
+	/**
+	 * Set the TimeOfDayPreview and broadcast the event.
+	 * @param InHours preview time to set in hours.
+	 */
+	UFUNCTION(BlueprintCallable, Category="TimeOfDay")
+	void SetTimeOfDayPreview(float InHours);
+
+	/**
+	 * Get the length of each day in hours.
+	 *
+	 * @return float, time in hours.
+	 */
+	UFUNCTION(BlueprintCallable, Category="TimeOfDay")
+	float GetDayLength() const;
+
+	/**
+	 * Set the duration of the day in game time in hours.
+	 *
+	 * This is currently intentionally not exposed to Blueprint since
+	 * this property can only be realized by rebuilding the root sequence.
+	 *
+	 * @param InHours day length in hours.
+	 */
+	void SetDayLength(float InHours);
+
+	/**
+	 * Get the duration of each day cycle in hours.
+	 *
+	 * @return float, time in hours
+	 */
+	UFUNCTION(BlueprintCallable, Category="TimeOfDay")
+	float GetTimePerCycle() const;
+
+	/**
+	 * Set the duration of a day cycle in real time in hours.
+	 *
+	 * This is currently intentionally not exposed to Blueprint since
+	 * this property can only be realized by rebuilding the root sequence.
+	 *
+	 * @param InHours time per cycle in hours.
+	 */
+	void SetTimePerCycle(float InHours);
+
+	/**
+	 * Set the duration of a day cycle in real time in hours and update
+	 * the root sequence in game.
+	 *
+	 * This method is the the network multicast variant of SetTimePerCycle.
+	 * In game, this will properly update the root sequence and subsequence
+	 * time scales on both the server and the client. These changes must
+	 * be invoked on both server and client since only playback position &
+	 * status are replicated.
+	 *
+	 * @param InHours time per cycle in hours.
+	 */
+	UFUNCTION(NetMulticast, Reliable)
+	void Multicast_SetTimePerCycle(float InHours);
+
+	/**
+	 * Get the initial time of day in hours.
+	 *
+	 * @return float, time in hours
+	 */
+	UFUNCTION(BlueprintCallable, Category="TimeOfDay")
+	float GetInitialTimeOfDay() const;
+
+	/**
+	 * Set the initial time of day in game time in hours.
+	 *
+	 * This is currently intentionally not exposed to Blueprint since
+	 * this property can only be realized by rebuilding the root sequence.
+	 *
+	 * @param InHours initial time of day.
+	 */
+	void SetInitialTimeOfDay(float InHours);
+
+	/**
+	 * Get the current time of day in hours.
+	 *
+	 * @return float, time in hours
+	 */
+	UFUNCTION(BlueprintCallable, Category="TimeOfDay")
+	float GetTimeOfDay() const;
+
+	/**
+	 * Set the current time of day in hours (server only).
+	 *
+	 * @param InHours time of day in hours
+	 */
+	UFUNCTION(BlueprintCallable, Category="TimeOfDay")
+	bool SetTimeOfDay(float InHours);
+
+	/** Resume playback of the sequence. */
+	UFUNCTION(BlueprintCallable, Category="Playback")
+	void Play();
+
+	/** Pause playback of the sequence. */
+	UFUNCTION(BlueprintCallable, Category="Playback")
+	void Pause();
+
+	/** @return true if the sequence is playing. */
+	UFUNCTION(BlueprintCallable, Category="Playback")
+	bool IsPlaying() const;
+
+	/** @return true if the sequence is paused. */
+	UFUNCTION(BlueprintCallable, Category="Playback")
+	bool IsPaused() const;
+
+	/**
+	 * Check whether this day sequence actor has a static (fixed) time-of-day
+	 */
+	UFUNCTION(BlueprintCallable, Category="TimeOfDay")
+	bool HasStaticTimeOfDay() const;
+
+	/**
+	 * Get this day sequence actor's static (fixed) time-of-day, or numeric_limits<float>::lowest() if it doesn't have one
+	 */
+	UFUNCTION(BlueprintCallable, Category="TimeOfDay")
+	float GetStaticTimeOfDay() const;
+
+	/**
+	 * Force the Day Sequence to show a specific constant time (in hours)
+	 */
+	UFUNCTION(BlueprintCallable, Category="TimeOfDay")
+	void SetStaticTimeOfDay(float InHours);
+
+	/**
+	 * Revert the Day Sequence time back to its regular day cycle
+	 *
+	 * @param bResumeFromStaticTime  When checked, the day sequence will resume its playback from the static time,
+	 *                               otherwise it will restore to its natural position as if it had no static time applied.
+	 *                               Note: when this sequence is replicated, enabling this parameter will immediately be overwritten
+	 *                               by the server time on the next replication event.
+	 */
+	UFUNCTION(BlueprintCallable, Category="TimeOfDay")
+	void RemoveStaticTimeOfDay(bool bResumeFromStaticTime = false);
+
+	UDaySequence* GetRootSequence() const;
+
+	virtual void UpdateRootSequence();
+#if WITH_EDITOR
+	void UpdateRootSequenceOnTick();
+#endif
+
+	bool HasValidRootSequence() const;
+	bool RootSequenceHasValidSections() const;
+
+
+	DECLARE_EVENT(ADaySequenceActor, FOnRootSequenceChanged)
+	FOnRootSequenceChanged& GetOnPostRootSequenceChanged() { return OnPostRootSequenceChanged; }
+	FOnRootSequenceChanged& GetOnPreRootSequenceChanged() { return OnPreRootSequenceChanged; }
+
+	/**
+	 * This delegate is broadcast after the Day Sequence Actor has initialized all of its own sequences to allow
+	 * other sequence providers to initialize their own sequences. If external providers were previously initialized,
+	 * they should cache their created subsections and then check to see if they exist in this map. If they do, the bool
+	 * associated with that subsection in this map can be set to true and it will be preserved. All subsections that map
+	 * to false will be deleted from the root sequence shortly after this delegate is broadcast.
+	*/
+	using FSubSectionPreserveMap = TMap<UMovieSceneSubSection*, bool>;
+	DECLARE_EVENT_OneParam(ADaySequenceActor, FOnPostInitializeDaySequences, FSubSectionPreserveMap*)
+	FOnPostInitializeDaySequences& GetOnPostInitializeDaySequences() { return OnPostInitializeDaySequences; }
+
+	void InvalidateMuteStates() const;
+	
+#if ENABLE_DRAW_DEBUG
+	const FName ShowDebug_GeneralCategory = "DaySequence";
+	const FName ShowDebug_SubSequenceCategory = "DaySequenceSubSequences";
+	
+	DECLARE_EVENT_OneParam(ADaySequenceActor, FOnDebugLevelChanged, int32)
+	FOnDebugLevelChanged& GetOnDebugLevelChanged() { return OnDebugLevelChanged; }
+
+	bool IsDebugCategoryRegistered(const FName& Category) const;
+	void RegisterDebugCategory(const FName& Category, UE::DaySequence::FDebugCategoryDrawFunction DrawFunction);
+
+	void RegisterDebugEntry(const TWeakPtr<UE::DaySequence::FDaySequenceDebugEntry>& DebugEntry, const FName& Category);
+	void UnregisterDebugEntry(const TWeakPtr<UE::DaySequence::FDaySequenceDebugEntry>& DebugEntry, const FName& Category);
+
+	// This will be called once per category and is passed the data for all entries associated with this category
+	static void OnShowDebugInfoDrawFunction(UCanvas* Canvas, TArray<TSharedPtr<TMap<FString, FString>>>& Entries, const FString& Category);
+#endif
+	
+	/**
+	 * Normalizes a subsection's timescale and frame range to the root sequence such
+	 * that the subsection's sequence represents a full day cycle.
+	 */
+	void UpdateSubSectionTimeScale(UMovieSceneSubSection* InSubSection) const;
+	
+	/**
+	 * TODO [nickolas.drake]
+	 * Write a class that wraps TrackConditionMap (TrackConditionMapProxy?) and make
+	 * EvaluateSequenceConditions and BindToConditionCallbacks member functions.
+	 */
+
+	/* Returns a TObjectPtr which is valid if ConditionTagClass is a subclass of UDaySequenceConditionTag and nullptr otherwise. */
+	TObjectPtr<UDaySequenceConditionTag> GetOrInstantiateConditionTag(const TSubclassOf<UDaySequenceConditionTag>& ConditionClass);
+	
+	/* Called to evaluate a particular set of conditions. Instantiates conditions if necessary. TODO: remove bInitialMuteState */
+	bool EvaluateSequenceConditions(bool bInitialMuteState, const FDaySequenceConditionSet::FConditionValueMap& InConditions);
+
+	/* Called to register InFunction as a callback for instances matching InConditions with a lifetime equivalent to LifetimeObject. */
+	void BindToConditionCallbacks(UObject* LifetimeObject, const FDaySequenceConditionSet::FConditionValueMap& InConditions, const TFunction<void(void)>& InFunction);
+	
+protected:
+	//~ Begin UObject interface
+	virtual void PostInitProperties() override;
+	virtual void PostLoad() override;
+	//~ End UObject interface
+	
+	//~ Begin AActor interface
+	virtual void PostInitializeComponents() override;
+	virtual void BeginPlay() override;
+	virtual void EndPlay(const EEndPlayReason::Type EndPlayReason) override;
+	virtual void RewindForReplay() override;
+	virtual void Destroyed() override;
+
+#if WITH_EDITOR
+	virtual bool CanChangeIsSpatiallyLoadedFlag() const override { return false; }
+#endif
+	
+#if WITH_EDITOR || ENABLE_DRAW_DEBUG	// This expression redundantly checks WITH_EDITOR but is expressing the intended behavior
+	virtual void OnConstruction(const FTransform& Transform) override;
+	virtual void Tick(float DeltaTime) override;
+	virtual bool ShouldTickIfViewportsOnly() const override;
+#endif
+	//~ End AActor interface
+
+	//~ Begin IMovieScenePlaybackClient interface
+	virtual bool RetrieveBindingOverrides(const FGuid& InBindingId, FMovieSceneSequenceID InSequenceID, TArray<UObject*, TInlineAllocator<1>>& OutObjects) const override;
+	virtual UObject* GetInstanceData() const override;
+	virtual bool GetIsReplicatedPlayback() const override;
+	virtual void WarpEvaluationRange(FMovieSceneEvaluationRange& InOutRange) const override;
+	//~ End IMovieScenePlaybackClient interface
+
+#if WITH_EDITOR
+	//~ Begin IMovieSceneBindingOwnerInterface
+	virtual TSharedPtr<FStructOnScope> GetObjectPickerProxy(TSharedPtr<IPropertyHandle> PropertyHandle) override;
+	virtual void UpdateObjectFromProxy(FStructOnScope& Proxy, IPropertyHandle& ObjectPropertyHandle) override;
+	virtual UMovieSceneSequence* RetrieveOwnedSequence() const override;
+	//~ End IMovieSceneBindingOwnerInterface
+	
+	virtual void PostEditChangeProperty(FPropertyChangedEvent& PropertyChangedEvent) override;
+#endif //WITH_EDITOR
+
+	/** Initialize SequencePlayer with transient root sequence */
+	void InitializePlayer();
+	void InitializeRootSequence();
+
+#if ROOT_SEQUENCE_RECONSTRUCTION_ENABLED
+	/**
+	 * Determines if subsection initialization is necessary or not by verifying all sections in SubSections exist in this map.
+	 * If they do, we keep all existing subsections. If not, we remove existing subsections and call InitializeDaySequences().
+	 * 
+	 * @param SectionsToPreserve A map of RootSequence subsections that allows sequence providers to mark if their cached subsections should be kept or removed.
+	 * @return Returns true if we need to remove all sections and reinitialize, false if we can keep all existing sections and bypass initialization.
+	 */
+	bool MarkDaySequences(FSubSectionPreserveMap* SectionsToPreserve);
+#endif
+	
+	/**
+	 * Called by InitializeRootSequence.
+	 * By default this will initialize sequences in DaySequenceCollection.
+	 * Non-collection sequences should be initialized in an override of this function in order to be properly inspected by the debug overlay.
+	 */
+	virtual void InitializeDaySequences();
+
+	UMovieSceneSubSection* InitializeDaySequence(const FDaySequenceCollectionEntry& SequenceAsset);
+
+	/** Compute PlaybackSettings from day cycle properties */
+	FMovieSceneSequencePlaybackSettings GetPlaybackSettings(const UDaySequence* Sequence) const;
+
+	void OnSequencePlayerUpdate(const UMovieSceneSequencePlayer& Player, FFrameTime CurrentTime, FFrameTime PreviousTime);
+	virtual void SequencePlayerUpdated(float CurrentTime, float PreviousTime);
+
+public:
+	/**
+	 * User-provided interpolation curve that maps day cycle times to desired cycle times (usually from 0-24 hours).
+	 * When disabled, the cycle will interpolate linearly.
+	 **/
+	UPROPERTY(EditAnywhere, replicated, BlueprintReadWrite, Category=RuntimeDayCycle, meta=(EditCondition=bUseInterpCurve, DisplayAfter=bRunDayCycle))
+	TObjectPtr<UCurveFloat> DayInterpCurve;
+	
+	UPROPERTY(EditAnywhere, Category="Sequence")
+	TObjectPtr<UDaySequenceCollectionAsset> DaySequenceCollection;
+
+	/** User-defined bias to apply to sequences in DaySequenceCollection. */
+	UPROPERTY(BlueprintReadWrite, EditAnywhere, Category="Sequence", DisplayName="Collection Bias", meta=(DisplayAfter="DaySequenceCollection"))
+	int32 Bias;
+		
+	/** Mapping of actors to override the sequence bindings with */
+	UPROPERTY(Instanced, BlueprintReadOnly, Category="Sequence")
+	TObjectPtr<UMovieSceneBindingOverrides> BindingOverrides;
+	
+	/** If true, playback of this sequence on the server will be synchronized across other clients */
+	UPROPERTY(EditAnywhere, DisplayName="Replicate Playback", BlueprintReadWrite, BlueprintSetter=SetReplicatePlayback, Category=Replication)
+	uint8 bReplicatePlayback : 1;
+
+	/** Used to early out from WarpEvaluationRange. Set to true when editing the root sequence of this actor. */
+	bool bForceDisableDayInterpCurve = false;
+
+#if WITH_EDITORONLY_DATA
+protected:
+	/**
+	 * Editor only override for testing in PIE.
+	 * If true, the actor uses the current preview time as the initial time of day.
+	 * Otherwise, the initial time of day is InitialTimeOfDay.
+	 */
+	UPROPERTY()
+	bool bOverrideInitialTimeOfDay;
+	
+	/**
+	 * Editor only override for testing in PIE.
+	 * If true, prevents initial playback on BeginPlay.
+	 * Otherwise, initial playback is determined by bRunDayCycle.
+	 */
+	UPROPERTY()
+	bool bOverrideRunDayCycle;
+
+public:
+	bool GetOverrideInitialTimeOfDay() const;
+	void SetOverrideInitialTimeOfDay(bool bNewOverrideInitialTimeOfDay);
+
+	// Variant of SetOverrideInitialTimeOfDay that will update TimeOfDayPreview if necessary.
+	void SetOverrideInitialTimeOfDay(bool bNewOverrideInitialTimeOfDay, float OverrideInitialTimeOfDay);
+
+	bool GetOverrideRunDayCycle() const;
+	void SetOverrideRunDayCycle(bool bNewOverrideRunDayCycle);
+	
+	DECLARE_MULTICAST_DELEGATE_TwoParams(FOnOverrideInitialTimeOfDayChanged, bool, float);
+	FOnOverrideInitialTimeOfDayChanged OnOverrideInitialTimeOfDayChanged;
+
+	DECLARE_MULTICAST_DELEGATE_OneParam(FOnOverrideRunDayCycleChanged, bool);
+	FOnOverrideRunDayCycleChanged OnOverrideRunDayCycleChanged;
+#endif
+	
+protected:
+	UPROPERTY(Instanced, transient, replicated, BlueprintReadOnly, BlueprintGetter=GetSequencePlayer, Category="Playback", meta=(ExposeFunctionCategories="Sequencer|Player"))
+	TObjectPtr<UDaySequencePlayer> SequencePlayer;
+	
+	UPROPERTY(Transient)
+	TObjectPtr<UDaySequence> RootSequence;
+
+#if WITH_EDITORONLY_DATA
+	/**
+	 * Sets the time of day to preview in the editor. Does not affect the start time at runtime.
+	 * Can be used as the initial time of day in PIE if bOverrideInitialTimeOfDay is true.
+	 */
+	UPROPERTY(EditAnywhere, Category = Preview, NonTransactional)
+	FTimecode TimeOfDayPreview;
+#endif
+
+	/** Whether or not to run a day cycle. If this is unchecked the day cycle will remain fixed at the time specified by the Initial Time setting */
+	UPROPERTY(EditAnywhere, Category = RuntimeDayCycle)
+	bool bRunDayCycle;
+
+	UPROPERTY(EditAnywhere, replicated, BlueprintReadWrite, Category = RuntimeDayCycle, meta=(InlineEditConditionToggle))
+	bool bUseInterpCurve;
+
+	/** How long a single day cycle is */
+	UPROPERTY(EditAnywhere, Category = RuntimeDayCycle)
+	FTimecode DayLength;
+
+	/** How long does it take for a day cycle to complete in world time. If this is the same value as day duration that means real world time is used */
+	UPROPERTY(EditAnywhere, Category = RuntimeDayCycle)
+	FTimecode TimePerCycle;
+
+	/** The initial time that the day cycle will start at */
+	UPROPERTY(EditAnywhere, Category = RuntimeDayCycle)
+	FTimecode InitialTimeOfDay;
+	
+	/** Cached static time controller applied from SetStaticTimeOfDay */
+	TWeakPtr<UE::DaySequence::FStaticTimeControllerOverride> WeakTimeControllerOverride;
+
+#if WITH_EDITORONLY_DATA
+	/**
+	 * Editor only static time, necessary because WeakTimeControllerOverride only exists in
+	 * game worlds and we have no way of tracking static time outside of game worlds otherwise.
+	 */
+	TOptional<float> StaticTime;
+	
+	/**
+	 * Blueprint exposed delegate invoked when the TimeOfDayPreview property
+	 * is changed.
+	 * @param new preview time in hours.
+	 */
+	UPROPERTY(BlueprintAssignable, Transient, meta=(AllowPrivateAccess="true"))
+	FOnTimeOfDayPreviewChanged OnTimeOfDayPreviewChanged;
+#endif //WITH_EDITORONLY_DATA
+
+#if WITH_EDITOR
+	/** When true, the root sequence will be updated on next Tick. */
+	bool bUpdateRootSequenceOnTick = false;
+#endif
+
+	FOnRootSequenceChanged OnPreRootSequenceChanged;
+	FOnRootSequenceChanged OnPostRootSequenceChanged;
+	FOnPostInitializeDaySequences OnPostInitializeDaySequences;
+	UE::DaySequence::FOnInvalidateMuteStates OnInvalidateMuteStates;
+	
+#if ENABLE_DRAW_DEBUG
+	void OnShowDebugInfo(AHUD* HUD, UCanvas* Canvas, const FDebugDisplayInfo& DisplayInfo, float& YL, float& YPos);
+
+	// The keys in this map are debug categories. When a debug entry is registered, it is is added
+	// to the array corresponding to the category it specified when registering. When OnShowDebugInfo is called on this
+	// actor, the debug data is consolidated and passed to the FDebugCategoryDrawFunction associated with the category
+	UE::DaySequence::FDebugEntryMap DebugEntries;
+
+	/** Debug entries this actor is responsible for. Populated by InitializeDaySequence and cleaned by InitializeRootSequence. */
+	TArray<TSharedPtr<UE::DaySequence::FDaySequenceDebugEntry>> SubSectionDebugEntries;
+	
+	// Notifies any listeners when the CVar "DaySequence.DebugLevel" changes and provides the new value.
+	FOnDebugLevelChanged OnDebugLevelChanged;
+	int32 CachedDebugLevel = 0;
+#endif
+
+	TArray<UMovieSceneSubSection*> SubSections;
+	
+	UPROPERTY(Transient)
+	TMap<TSubclassOf<UDaySequenceConditionTag>, TObjectPtr<UDaySequenceConditionTag>> TrackConditionMap;
+};
+
diff --git a/Engine/Plugins/Experimental/DaySequence/Source/DaySequence/Public/DaySequenceActorSpawner.h b/Engine/Plugins/Experimental/DaySequence/Source/DaySequence/Public/DaySequenceActorSpawner.h
new file mode 100644
--- /dev/null
+++ b/Engine/Plugins/Experimental/DaySequence/Source/DaySequence/Public/DaySequenceActorSpawner.h
@@ -0,0 +1,16 @@
+// Copyright Epic Games, Inc. All Rights Reserved.
+
+#pragma once
+
+#include "IMovieSceneObjectSpawner.h"
+
+class DAYSEQUENCE_API FDaySequenceActorSpawner : public IMovieSceneObjectSpawner
+{
+public:
+	static TSharedRef<IMovieSceneObjectSpawner> CreateObjectSpawner();
+
+	// IMovieSceneObjectSpawner interface
+	virtual UClass* GetSupportedTemplateType() const override;
+	virtual UObject* SpawnObject(FMovieSceneSpawnable& Spawnable, FMovieSceneSequenceIDRef TemplateID, TSharedRef<const UE::MovieScene::FSharedPlaybackState> SharedPlaybackState) override;
+	virtual void DestroySpawnedObject(UObject& Object) override;
+};
