DaySequence: Move DaySequence from Restricted/NFL to Experimental

diff --git a/Engine/Plugins/Experimental/DaySequence/Source/DaySequence/Public/DaySequenceModifierComponent.h b/Engine/Plugins/Experimental/DaySequence/Source/DaySequence/Public/DaySequenceModifierComponent.h
new file mode 100644
--- /dev/null
+++ b/Engine/Plugins/Experimental/DaySequence/Source/DaySequence/Public/DaySequenceModifierComponent.h
@@ -0,0 +1,509 @@
+// Copyright Epic Games, Inc. All Rights Reserved.
+
+#pragma once
+
+#include "Components/SceneComponent.h"
+#include "Tickable.h"
+#include "Components/BoxComponent.h"
+#include "Containers/Ticker.h"
+#include "DaySequenceActor.h"
+#include "DaySequenceConditionSet.h"
+#include "DrawDebugHelpers.h"	// Defines ENABLE_DRAW_DEBUG
+#include "Generators/MovieSceneEasingFunction.h"
+
+#include "DaySequenceModifierComponent.generated.h"
+
+class UMaterialInterface;
+class UShapeComponent;
+namespace ECollisionEnabled { enum Type : int; }
+namespace EEndPlayReason { enum Type : int; }
+struct FComponentReference;
+
+class UDaySequence;
+class UDaySequenceCollectionAsset;
+class UMovieSceneSubSection;
+
+struct FDaySequenceCollectionEntry;
+
+#if ENABLE_DRAW_DEBUG
+namespace UE::DaySequence
+{
+	struct FDaySequenceDebugEntry;
+}
+#endif
+
+/** Enum specifying how to control a day / night cycle from a modifier */
+UENUM(BlueprintType)
+enum class EDayNightCycleMode : uint8
+{
+	/** (default) Make no changes to the day/night cycle time */
+	Default,
+	/** Force the day/night cycle to be fixed at the specified constant time */
+	FixedTime,
+	/** Set an initial time for the day/night cycle when this modifier is enabled */
+	StartAtSpecifiedTime,
+	/** Use a random, fixed time for the day/night cycle */
+	RandomFixedTime,
+	/** Start the day/night cycle at a random time, and allow it to continue from there */
+	RandomStartTime,
+};
+
+/** Enum specifying how to blend into and out of a UDaySequenceModifierComponent that is using a trigger volume */
+UENUM(BlueprintType)
+enum class EDaySequenceModifierBlendMode : uint8
+{
+	/** Do not blend into the volume - modifier will immediately active when enabled */
+	None,
+	/** Blend the modifier in based on the distance within the volume. Only functional when EnableDistanceVolumeBlends is called. */
+	Distance,
+};
+
+#if WITH_EDITOR
+	/**
+	 * Editor-only tickable class that allows us to enable trigger volume previews based on
+	 * persepective camera position in the level viewport.
+	 */
+	struct FDaySequenceModifierComponentTickableBase : FTickableGameObject
+	{
+		virtual void UpdateEditorPreview(float DeltaTime)
+		{}
+
+		void Tick(float DeltaTime) override
+		{
+			/// Overridden here to work around ambiguous Tick function on USceneComponent
+			// Re-trigger the function as a differnt named virtual function
+			UpdateEditorPreview(DeltaTime);
+		}
+	};
+#else
+	/** Empty in non-editor builds */
+	struct FDaySequenceModifierComponentTickableBase {};
+#endif
+
+UCLASS()
+class UDaySequenceModifierEasingFunction
+	: public UObject
+	, public IMovieSceneEasingFunction
+{
+public:
+	GENERATED_BODY()
+	
+	enum class EEasingFunctionType
+	{
+		EaseIn,
+		EaseOut
+	};
+	
+	void Initialize(EEasingFunctionType EasingType);
+	
+	virtual float Evaluate(float Interp) const override;
+
+private:
+	TFunction<float(float)> EvaluateImpl;
+};
+
+DECLARE_DYNAMIC_MULTICAST_DELEGATE(FOnPostReinitializeSubSequences);
+DECLARE_DYNAMIC_MULTICAST_DELEGATE(FOnPostEnableModifier);
+DECLARE_DYNAMIC_MULTICAST_DELEGATE_OneParam(FOnVolumeBlendTargetOverlap, AActor*, OverlappingActor);
+DECLARE_MULTICAST_DELEGATE_OneParam(FOnVolumeBlendTargetOverlapNative, AActor*);
+
+UCLASS(MinimalAPI, BlueprintType, Blueprintable, config=Game, HideCategories=(Physics,Navigation,Collision,HLOD,Rendering,Cooking,Mobile,RayTracing,AssetUserData), meta=(BlueprintSpawnableComponent))
+class UDaySequenceModifierComponent
+	: public USceneComponent
+	, public FDaySequenceModifierComponentTickableBase
+{
+public:
+	GENERATED_BODY()
+
+	UDaySequenceModifierComponent(const FObjectInitializer& Init);
+
+#if WITH_EDITOR
+	DAYSEQUENCE_API static void SetVolumePreviewLocation(const FVector& Location);
+#endif
+
+	/**
+	 * Bind this component to the specified day sequence actor.
+	 * Will not add our overrides to the sub-sequence until EnableModifier is called.
+	 */
+	UFUNCTION(BlueprintCallable, Category="Day Sequence")
+	DAYSEQUENCE_API void BindToDaySequenceActor(ADaySequenceActor* DaySequenceActor);
+
+	/**
+	 * Unbind this component from its day sequence actor if valid.
+	 * Will removes the sub-sequence from the root sequence if it's set up.
+	 */
+	UFUNCTION(BlueprintCallable, Category="Day Sequence")
+	DAYSEQUENCE_API void UnbindFromDaySequenceActor();
+
+	/**
+	 * Enable this component.
+	 */
+	UFUNCTION(BlueprintCallable, Category="Day Sequence")
+	DAYSEQUENCE_API void EnableComponent();
+
+	/**
+	 * Disable this component.
+	 * Will remove the sub-sequence from the root sequence if it's set up.
+	 */
+	UFUNCTION(BlueprintCallable, Category="Day Sequence")
+	DAYSEQUENCE_API void DisableComponent();
+
+	/**
+	 * Enable this modifier by enabling its sub-sequence within the root Day Sequence
+	 * Will initialize a new sub-sequence track if necessary.
+	 */
+	UFUNCTION(BlueprintCallable, Category="Day Sequence")
+	DAYSEQUENCE_API void EnableModifier();
+
+	/**
+	 * Disable this modifier by disabling its sub-sequence track
+	 */
+	UFUNCTION(BlueprintCallable, Category="Day Sequence")
+	DAYSEQUENCE_API void DisableModifier();
+
+public:
+
+	/**
+	 * Remove all the Sequencer tracks within our procedural Day Sequence
+	 */
+	UFUNCTION(BlueprintCallable, Category="Day Sequence")
+	DAYSEQUENCE_API void ResetOverrides();
+
+	/**
+	 * Add a new override for the static time of day on the day-sequence actor
+	 *
+	 * @param Hours         The static time of day in hours to use
+	 */
+	UFUNCTION(BlueprintCallable, Category="Day Sequence")
+	DAYSEQUENCE_API void AddStaticTimeOfDayOverride(ADaySequenceActor* Actor, float Hours);
+
+	/**
+	 * Add a new boolean override for the specified property
+	 *
+	 * @param Object        The object that contains the property we want to override. Must either be the ADaySequenceActor or one of its components
+	 * @param PropertyName  The name of the property to override
+	 * @param Value         The value to override
+	 */
+	UFUNCTION(BlueprintCallable, Category="Day Sequence")
+	DAYSEQUENCE_API void AddBoolOverride(UObject* Object, FName PropertyName, bool bValue);
+
+	/**
+	 * Add a new scalar override for the specified property
+	 *
+	 * @param Object        The object that contains the property we want to override. Must either be the ADaySequenceActor or one of its components
+	 * @param PropertyName  The name of the property to override
+	 * @param Value         The value to override
+	 */
+	UFUNCTION(BlueprintCallable, Category="Day Sequence")
+	DAYSEQUENCE_API void AddScalarOverride(UObject* Object, FName PropertyName, double Value);
+
+	/**
+	 * Add a new vector override for the specified property
+	 *
+	 * @param Object        The object that contains the property we want to override. Must either be the ADaySequenceActor or one of its components
+	 * @param PropertyName  The name of the property to override
+	 * @param Value         The value to override
+	 */
+	UFUNCTION(BlueprintCallable, Category="Day Sequence")
+	DAYSEQUENCE_API void AddVectorOverride(UObject* Object, FName PropertyName, FVector Value);
+
+	/**
+	 * Add a new color override for the specified property
+	 *
+	 * @param Object        The object that contains the property we want to override. Must either be the ADaySequenceActor or one of its components
+	 * @param PropertyName  The name of the property to override.
+	 * @param Value         The value to override
+	 */
+	UFUNCTION(BlueprintCallable, Category="Day Sequence")
+	DAYSEQUENCE_API void AddColorOverride(UObject* Object, FName PropertyName, FLinearColor Value);
+
+	/**
+	 * Add a new transform override for the specified property
+	 *
+	 * @param Object        The object that contains the property we want to override. Must either be the ADaySequenceActor or one of its components
+	 * @param Value         The value to override
+	 */
+	UFUNCTION(BlueprintCallable, Category="Day Sequence")
+	DAYSEQUENCE_API void AddTransformOverride(UObject* Object, FTransform Value);
+
+	/**
+	 * Add a new material override for the specified material element index
+	 *
+	 * @param Object          The object that contains the property we want to override. Must either be the ADaySequenceActor or one of its components
+	 * @param MaterialIndex   The index of the material element to override
+	 * @param Value           The material to assign
+	 */
+	UFUNCTION(BlueprintCallable, Category="Day Sequence")
+	DAYSEQUENCE_API void AddMaterialOverride(UObject* Object, int32 MaterialIndex, UMaterialInterface* Value);
+
+	/**
+	 * Add a new scalar material parameter override for the specified material element index and parameter name
+	 *
+	 * @param Object          The object that contains the material we want to override. Must either be the ADaySequenceActor or one of its components
+	 * @param MaterialIndex   The index of the material element to override
+	 * @param ParameterName   The name of the parameter to override
+	 * @param Value           The scalar value to assign to the material parameter
+	 */
+	UFUNCTION(BlueprintCallable, Category="Day Sequence")
+	DAYSEQUENCE_API void AddScalarMaterialParameterOverride(UObject* Object, int32 MaterialIndex, FName ParameterName, float Value);
+
+	/**
+	 * Add a new color material parameter override for the specified material element index and parameter name
+	 *
+	 * @param Object          The object that contains the material we want to override. Must either be the ADaySequenceActor or one of its components
+	 * @param MaterialIndex   The index of the material element to override
+	 * @param ParameterName   The name of the parameter to override
+	 * @param Value           The color value to assign to the material parameter
+	 */
+	UFUNCTION(BlueprintCallable, Category="Day Sequence")
+	DAYSEQUENCE_API void AddColorMaterialParameterOverride(UObject* Object, int32 MaterialIndex, FName ParameterName, FLinearColor Value);
+
+	/**
+	 * Add a new visibility override for the specified property
+	 *
+	 * @param Object        The object that contains the property we want to override. Must either be the ADaySequenceActor or one of its components
+	 * @param Value         The value to override
+	 */
+	UFUNCTION(BlueprintCallable, Category="Day Sequence")
+	DAYSEQUENCE_API void AddVisibilityOverride(UObject* Object, bool bValue);
+
+	/**
+	 * Sets the user day sequence
+	 */
+	UFUNCTION(BlueprintCallable, Category="Day Sequence")
+	DAYSEQUENCE_API void SetUserDaySequence(UDaySequence* InDaySequence);
+
+	/**
+	 * Sets whether to enable modifier based on volume
+	 */
+	DAYSEQUENCE_API void SetUseVolume(bool bState);
+
+	/**
+	 * Enable distance-based volume blends based on the location of the specified actor relative to our volume.
+	 * FinalWeight = Min(DistanceVolumeBlendWeight, CustomVolumeBlendWeight).
+	 */
+	UFUNCTION(BlueprintCallable, Category="Day Sequence")
+	DAYSEQUENCE_API void EnableDistanceVolumeBlends(AActor* InActor);
+
+	/**
+	 * Determines the number of volumes the current blend target is
+	 * overlapping and returns true if that number is greater than 0
+	 */
+	UFUNCTION(BlueprintCallable, Category="Day Sequence")
+	DAYSEQUENCE_API bool IsBlendTargetInAnyVolume();
+	
+	/**
+	 * Sets a custom blend weight for volume based blends.
+	 * FinalWeight = Min(DistanceVolumeBlendWeight, CustomVolumeBlendWeight).
+	 */
+	UFUNCTION(BlueprintCallable, Category="Day Sequence")
+	DAYSEQUENCE_API void SetCustomVolumeBlendWeight(float Weight);
+
+	/**
+	 * Get the current blend weight.
+	 */
+	UFUNCTION(BlueprintCallable, Category="Day Sequence")
+	float GetCurrentBlendWeight() const;
+
+	/**
+	 * Applies the given collision type to all valid volume shape components.
+	 */
+	UFUNCTION(BlueprintCallable, Category="Day Sequence")
+	void SetVolumeCollisionEnabled(const ECollisionEnabled::Type InCollisionType) const;
+
+	void EmptyVolumeShapeComponents();
+	void AddVolumeShapeComponent(const FComponentReference& InShapeReference);
+
+	void InvalidateMuteStates() const;
+	
+	FOnVolumeBlendTargetOverlapNative& GetOnVolumeBlendTargetOverlapBegin() { return OnVolumeBlendTargetOverlapBeginNative; }
+	FOnVolumeBlendTargetOverlapNative& GetOnVolumeBlendTargetOverlapEnd() { return OnVolumeBlendTargetOverlapEndNative; }
+
+#if ENABLE_DRAW_DEBUG
+	bool ShouldShowDebugInfo() const;
+#endif
+	
+public:
+
+	/*~ Begin UActorComponent interface */
+	void BeginPlay() override;
+	void TickComponent(float DeltaTime, ELevelTick TickType, FActorComponentTickFunction* ThisTickFunction) override;
+	void EndPlay(EEndPlayReason::Type Reason) override;
+	void OnRegister() override;
+	void OnUnregister() override;
+	/*~ End UActorComponent interface */
+
+#if WITH_EDITOR
+	/*~ Begin FTickableGameObject interface */
+	void UpdateEditorPreview(float DeltaTime) override;
+	TStatId GetStatId() const override;
+	ETickableTickType GetTickableTickType() const override;
+	bool IsTickableWhenPaused() const override { return true; }
+	bool IsTickableInEditor() const override { return true; }
+	bool IsTickable() const override { return true; }
+	/*~ End FTickableGameObject interface */
+
+	void PostEditChangeProperty(FPropertyChangedEvent& PropertyChangedEvent) override;
+
+#endif
+
+protected:
+
+	virtual bool CanBeEnabled() const;
+
+	UPROPERTY(Instanced, Transient)
+	TObjectPtr<UDaySequenceModifierEasingFunction> EasingFunction;
+
+	TArray<UShapeComponent*> GetVolumeShapeComponents() const;
+
+protected:
+
+	void SetInitialTimeOfDay();
+
+protected:
+
+	/** Non-serialized target actor we are currently bound to */
+	UPROPERTY(Transient, BlueprintReadWrite, Category="Day Sequence")
+	TObjectPtr<ADaySequenceActor> TargetActor;
+
+	/** When set, the shape components will be used for the modifier volume, otherwise the default Box component will be used. */
+	UPROPERTY(BlueprintReadWrite, EditAnywhere, Category="Day Sequence", meta=(UseComponentPicker, AllowedClasses="/Script/Engine.ShapeComponent", AllowAnyActor))
+	TArray<FComponentReference> VolumeShapeComponents;
+
+	/** The actor to use for distance-based volume blend calculations */
+	UPROPERTY(Transient)
+	TWeakObjectPtr<AActor> ExternalVolumeBlendTarget;
+
+	/** An optional user-provided Day Sequence - used instead of our procedurally generated one if set */
+	UPROPERTY(BlueprintReadWrite, EditAnywhere, Category="Day Sequence", meta=(EditCondition="!bUseCollection", EditConditionHides, DisplayAfter="bUseCollection"))
+	TObjectPtr<UDaySequence> UserDaySequence;
+
+	/** The procedurally generated sequence containing our override tracks. Owned by this component. */
+	UPROPERTY(Transient)
+	TObjectPtr<UDaySequence> ProceduralDaySequence;
+
+	UPROPERTY(EditAnywhere, Category="Day Sequence", meta=(EditCondition="bUseCollection", EditConditionHides, DisplayAfter="bUseCollection"))
+	TObjectPtr<UDaySequenceCollectionAsset> DaySequenceCollection;
+	
+	/** User-defined bias. */
+	UPROPERTY(BlueprintReadWrite, EditAnywhere, Category="Day Sequence", meta=(EditCondition="!bIgnoreBias"))
+	int32 Bias;
+
+	/** The time to use for the day/night cycle. */
+	UPROPERTY(BlueprintReadWrite, EditAnywhere, Category="Time", DisplayName="Time", meta=(DisplayAfter="DayNightCycle", EditCondition="DayNightCycle==EDayNightCycleMode::FixedTime || DayNightCycle==EDayNightCycleMode::StartAtSpecifiedTime", EditConditionHides))
+	float DayNightCycleTime;
+
+	/** Defines the blend range based on BlendMode. When BlendMode == Distance, this is interpreted as world units */
+	UPROPERTY(BlueprintReadWrite, EditAnywhere, Category="Day Sequence", meta=(DisplayAfter="BlendMode", EditCondition="BlendMode!=EDaySequenceModifierBlendMode::None", EditConditionHides))
+	float BlendAmount;
+
+	/** Cached blend factor last time we ticked in the range 0-1. Used to prevent superfluous updates on a paused day sequence. */
+	float CachedBlendFactor;
+
+	/** Defines a custom blend weight. The final blend weight is the minimum between the Distance blend weight and CustomVolumeBlendWeight. */
+	UPROPERTY(Transient)
+	float CustomVolumeBlendWeight;
+
+	/** Changes the way this modifier controls the day/night cycle time */
+	UPROPERTY(BlueprintReadWrite, EditAnywhere, Category="Time", DisplayName="Day/Night Cycle")
+	EDayNightCycleMode DayNightCycle;
+
+	/** Changes the way this modifier blends in and out */
+	UPROPERTY(BlueprintReadWrite, EditAnywhere, Category="Day Sequence")
+	EDaySequenceModifierBlendMode BlendMode;
+
+	/** Blueprint exposed delegate invoked after the component's subsequences are reinitialized. */
+	UPROPERTY(BlueprintAssignable, Transient, meta=(AllowPrivateAccess="true"))
+	FOnPostReinitializeSubSequences OnPostReinitializeSubSequences;
+
+	/** Blueprint exposed delegate invoked after the modifier component is enabled. */
+	UPROPERTY(BlueprintAssignable, Transient, meta=(AllowPrivateAccess="true"))
+	FOnPostEnableModifier OnPostEnableModifier;
+
+	/** 
+	 * Blueprint exposed delegate invoked after the current blend target enters a volume and was not already in a volume.
+	 * Only broadcasts when owning actor does not have network authority.
+	 */
+	UPROPERTY(BlueprintAssignable, Transient, meta=(AllowPrivateAccess="true"))
+	FOnVolumeBlendTargetOverlap OnVolumeBlendTargetOverlapBegin;
+	FOnVolumeBlendTargetOverlapNative OnVolumeBlendTargetOverlapBeginNative;
+	
+	/**
+	 * Blueprint exposed delegate invoked after the current blend target has exited all volumes.
+	 * Only broadcasts when owning actor does not have network authority.
+	 */
+	UPROPERTY(BlueprintAssignable, Transient, meta=(AllowPrivateAccess="true"))
+	FOnVolumeBlendTargetOverlap OnVolumeBlendTargetOverlapEnd;
+	FOnVolumeBlendTargetOverlapNative OnVolumeBlendTargetOverlapEndNative;
+	
+	/** When enabled, these overrides will always override all settings regardless of their bias */
+	UPROPERTY(BlueprintReadWrite, EditAnywhere, Category="Day Sequence", meta=(InlineEditConditionToggle))
+	uint8 bIgnoreBias : 1;
+
+	/** Whether to use the bounding volume for enabling this modifier. If not enabled this modifier will always be active. */
+	UPROPERTY(BlueprintReadWrite, EditAnywhere, Category="Day Sequence")
+	uint8 bUseVolume : 1;
+
+	/** Flag used track whether or not this component is enabled or disabled. */
+	UPROPERTY(Transient, DuplicateTransient, BlueprintReadOnly, Category="Day Sequence")
+	uint8 bIsComponentEnabled : 1;
+	
+	/** Non-serialized variable for tracking whether our overrides are enabled or not */
+	UPROPERTY(Transient, DuplicateTransient, BlueprintReadOnly, Category="Day Sequence")
+	uint8 bIsEnabled : 1;
+
+	/** Preview this day sequence modifier in the editor */
+	UPROPERTY(Transient, BlueprintReadWrite, EditAnywhere, Category="Day Sequence")
+	uint8 bPreview : 1;
+
+	/** Flag to keep track of whether we need to unpause the day sequence when we are disabled */
+	uint8 bUnpauseOnDisable : 1;
+
+	/** If true, hide UserDaySequence and expose DaySequenceCollection. */
+	UPROPERTY(EditAnywhere, Category="Day Sequence")
+	uint8 bUseCollection : 1;
+
+private:
+	float GetDistanceBlendFactor(FVector Position) const;
+
+	/**
+	 * Creates and adds or marks for preserve all subsections that this modifier is responsible for.
+	 * Optionally provided a map of all sections that exist in the root sequence to a bool flag used to mark that section as still relevant.
+	 */
+	void ReinitializeSubSequence(ADaySequenceActor::FSubSectionPreserveMap* SectionsToPreserve);
+	UMovieSceneSubSection* InitializeDaySequence(const FDaySequenceCollectionEntry& SequenceAsset);
+	void RemoveSubSequenceTrack();
+	FGuid GetOrCreateProceduralBinding(UObject* Object);
+
+	void UpdateCachedExternalShapes() const;
+	mutable TArray<TWeakObjectPtr<UShapeComponent>> CachedExternalShapes;
+	mutable bool bCachedExternalShapesInvalid = true;
+
+	void BindOverlapEvents() const;
+	void UnbindOverlapEvents() const;
+	int32 OccupiedVolumes = 0;
+	
+	UFUNCTION()
+	void OnVolumeOverlapBegin(UPrimitiveComponent* OverlappedComponent, AActor* OtherActor, UPrimitiveComponent* OtherComp, int32 OtherBodyIndex, bool bFromSweep, const FHitResult& SweepResult);
+
+	UFUNCTION()
+	void OnVolumeOverlapEnd(UPrimitiveComponent* OverlappedComponent, AActor* OtherActor, UPrimitiveComponent* OtherComp, int32 OtherBodyIndex);
+	
+	/*~ Transient state for active gameplay */
+	TWeakObjectPtr<UMovieSceneSubSection> WeakSubSection;
+	TArray<TWeakObjectPtr<UMovieSceneSubSection>> SubSections;
+
+	UE::DaySequence::FOnInvalidateMuteStates OnInvalidateMuteStates;
+	
+#if ENABLE_DRAW_DEBUG
+	const FName ShowDebug_ModifierCategory = "DaySequence_Modifiers";
+	
+	void OnDebugLevelChanged(int32 InDebugLevel);
+	
+	/** Determines whether or not this modifier will show debug info */
+	int32 DebugLevel;
+
+	TSharedPtr<UE::DaySequence::FDaySequenceDebugEntry> DebugEntry;
+	TArray<TSharedPtr<UE::DaySequence::FDaySequenceDebugEntry>> SubSectionDebugEntries;
+#endif
+};
diff --git a/Engine/Plugins/Experimental/DaySequence/Source/DaySequence/Public/DaySequenceProvider.h b/Engine/Plugins/Experimental/DaySequence/Source/DaySequence/Public/DaySequenceProvider.h
new file mode 100644
--- /dev/null
+++ b/Engine/Plugins/Experimental/DaySequence/Source/DaySequence/Public/DaySequenceProvider.h
@@ -0,0 +1,26 @@
+// Copyright Epic Games, Inc. All Rights Reserved.
+
+#pragma once
+
+#include "GameFramework/Actor.h"
+#include "DaySequenceProvider.generated.h"
+
+class UDaySequence;
+
+UCLASS()
+class ADaySequenceProvider : public AActor
+{
+	GENERATED_BODY()
+public:
+	ADaySequenceProvider(const FObjectInitializer& Init);
+	
+	virtual TArrayView<TObjectPtr<UDaySequence>> GetDaySequences();
+
+#if WITH_EDITOR
+	virtual void PostEditChangeProperty(FPropertyChangedEvent& PropertyChangedEvent) override;
+#endif //WITH_EDITOR
+
+protected:
+	UPROPERTY(EditAnywhere, BlueprintReadOnly, Category="Sequence", meta=(AllowedClasses="/Script/DaySequence.DaySequence"))
+	TArray<TObjectPtr<UDaySequence>> DaySequenceAssets;
+};
diff --git a/Engine/Plugins/Experimental/DaySequence/Source/DaySequence/Public/DaySequenceSpawnRegister.h b/Engine/Plugins/Experimental/DaySequence/Source/DaySequence/Public/DaySequenceSpawnRegister.h
new file mode 100644
--- /dev/null
+++ b/Engine/Plugins/Experimental/DaySequence/Source/DaySequence/Public/DaySequenceSpawnRegister.h
@@ -0,0 +1,30 @@
+// Copyright Epic Games, Inc. All Rights Reserved.
+
+#pragma once
+
+#include "MovieSceneSpawnRegister.h"
+#include "UObject/Class.h"
+
+class IMovieScenePlayer;
+class IMovieSceneObjectSpawner;
+class UMovieSceneSpawnableBindingBase;
+
+/** Movie scene spawn register that knows how to handle spawning objects (actors) for a DaySequence  */
+class DAYSEQUENCE_API FDaySequenceSpawnRegister : public FMovieSceneSpawnRegister
+{
+public:
+	FDaySequenceSpawnRegister();
+
+protected:
+	/** ~ FMovieSceneSpawnRegister interface */
+	virtual UObject* SpawnObject(FMovieSceneSpawnable& Spawnable, FMovieSceneSequenceIDRef TemplateID, TSharedRef<const FSharedPlaybackState> SharedPlaybackState) override;
+	virtual void DestroySpawnedObject(UObject& Object, UMovieSceneSpawnableBindingBase* CustomSpawnableBinding) override;
+
+#if WITH_EDITOR
+	virtual bool CanSpawnObject(UClass* InClass) const override;
+#endif
+
+protected:
+	/** Extension object spawners */
+	TArray<TSharedRef<IMovieSceneObjectSpawner>> MovieSceneObjectSpawners;
+};
diff --git a/Engine/Plugins/Experimental/DaySequence/Source/DaySequence/Public/EnvironmentLightingActor.h b/Engine/Plugins/Experimental/DaySequence/Source/DaySequence/Public/EnvironmentLightingActor.h
new file mode 100644
--- /dev/null
+++ b/Engine/Plugins/Experimental/DaySequence/Source/DaySequence/Public/EnvironmentLightingActor.h
@@ -0,0 +1,47 @@
+// Copyright Epic Games, Inc. All Rights Reserved.
+
+#pragma once
+
+#include "DaySequenceProvider.h"
+#include "EnvironmentLightingActor.generated.h"
+
+class USkyAtmosphereComponent;
+class USkyLightComponent;
+class UDirectionalLightComponent;
+class UExponentialHeightFogComponent;
+class UVolumetricCloudComponent;
+
+UCLASS()
+class AEnvironmentLightingActor : public ADaySequenceProvider
+{
+	GENERATED_BODY()
+public:
+	AEnvironmentLightingActor(const FObjectInitializer& Init);
+
+protected:
+	// Components
+	UPROPERTY(Category=Environment, VisibleAnywhere, BlueprintReadOnly, meta=(AllowPrivateAccess="true", NoEditInline))
+	TObjectPtr<USkyAtmosphereComponent> SkyAtmosphereComponent;
+
+	UPROPERTY(Category=Environment, VisibleAnywhere, BlueprintReadOnly, meta=(AllowPrivateAccess="true", NoEditInline))
+	TObjectPtr<USkyLightComponent> SkyLightComponent;
+
+	UPROPERTY(Category=Environment, VisibleAnywhere, BlueprintReadOnly, meta=(AllowPrivateAccess="true", NoEditInline))
+	TObjectPtr<USceneComponent> SunRootComponent;
+
+	UPROPERTY(Category=Environment, VisibleAnywhere, BlueprintReadOnly, meta=(AllowPrivateAccess="true", NoEditInline))
+	TObjectPtr<UDirectionalLightComponent> SunComponent;
+
+	UPROPERTY(Category=Environment, VisibleAnywhere, BlueprintReadOnly, meta=(AllowPrivateAccess="true", NoEditInline))
+	TObjectPtr<UExponentialHeightFogComponent> ExponentialHeightFogComponent;
+
+	UPROPERTY(Category=Environment, VisibleAnywhere, BlueprintReadOnly, meta=(AllowPrivateAccess="true", NoEditInline))
+	TObjectPtr<UVolumetricCloudComponent> VolumetricCloudComponent;
+
+	static FName SkyAtmosphereComponentName;
+	static FName SkyLightComponentName;
+	static FName SunRootComponentName;
+	static FName SunComponentName;
+	static FName ExponentialHeightFogComponentName;
+	static FName VolumetricCloudComponentName;
+};
diff --git a/Engine/Plugins/Experimental/DaySequence/Source/DaySequence/Public/ProceduralDaySequenceBuilder.h b/Engine/Plugins/Experimental/DaySequence/Source/DaySequence/Public/ProceduralDaySequenceBuilder.h
new file mode 100644
--- /dev/null
+++ b/Engine/Plugins/Experimental/DaySequence/Source/DaySequence/Public/ProceduralDaySequenceBuilder.h
@@ -0,0 +1,83 @@
+// Copyright Epic Games, Inc. All Rights Reserved.
+
+#pragma once
+
+#include "CoreMinimal.h"
+#include "DaySequence.h"
+#include "DaySequenceActor.h"
+#include "Curves/RealCurve.h"
+
+#include "ProceduralDaySequenceBuilder.generated.h"
+
+/**
+ * A utility class for creating procedural Day Sequences.
+ * Before adding any keys, SetActiveBoundObject should be called and provided a Day Sequence Actor or a component owned by a Day Sequence Actor.
+ * All time values are currently normalized to the range [0, 1], inclusive on both ends. A time of 1 is handled as a special case and maps to the final frame.
+ * This class assumes the target Day Sequence Actor will stay alive and that users will keep the generated sequence alive, it manages no lifetimes.
+ */
+UCLASS(BlueprintType)
+class DAYSEQUENCE_API UProceduralDaySequenceBuilder
+	: public UObject
+{
+	GENERATED_BODY()
+	
+public:
+	/**
+	 * Initialize the procedural sequence and set the TargetActor for this builder.
+	 * 
+	 * @param InActor The target DaySequenceActor that will be animated by the generated sequence.
+	 * @param InitialSequence Optional sequence that this builder can operate on instead of allocating a new sequence.
+	 * @param bClearInitialSequence If true, calls ClearAllKeys().
+	 * @return The sequence which will be modified when calling SetActiveBoundObject and the Add*Key(s) functions.
+	 */
+	UFUNCTION(BlueprintCallable, Category = "Day Sequence")
+	UDaySequence* Initialize(ADaySequenceActor* InActor, UDaySequence* InitialSequence = nullptr, bool bClearInitialSequence = true);
+
+	/** Returns true Initialize has been called with a valid actor. */
+	UFUNCTION(BlueprintCallable, Category = "Day Sequence")
+	bool IsInitialized() const;
+	
+	/** Prepare the builder to begin adding keys animating properties on InObject. */
+	UFUNCTION(BlueprintCallable, Category = "Day Sequence")
+	void SetActiveBoundObject(UObject* InObject);
+
+	UFUNCTION(BlueprintCallable, Category = "Day Sequence")
+	void ClearAllKeys();
+	
+	UFUNCTION(BlueprintCallable, Category = "Day Sequence")
+	void AddScalarKey(FName PropertyName, float Key, double Value, ERichCurveInterpMode InterpMode = RCIM_Cubic);
+	void AddScalarKey(FName PropertyName, const TPair<float, double>& KeyValue, ERichCurveInterpMode InterpMode = RCIM_Cubic);
+	void AddScalarKeys(FName PropertyName, const TArray<TPair<float, double>>& KeysAndValues, ERichCurveInterpMode InterpMode = RCIM_Cubic);
+	
+	UFUNCTION(BlueprintCallable, Category = "Day Sequence")
+	void ClearScalarKeys(FName PropertyName);
+
+	UFUNCTION(BlueprintCallable, Category = "Day Sequence")
+	void AddBoolKey(FName PropertyName, float Key, bool Value);
+	void AddBoolKey(FName PropertyName, const TPair<float, bool>& KeyValue);
+	void AddBoolKeys(FName PropertyName, const TArray<TPair<float, bool>>& KeysAndValues);
+
+	UFUNCTION(BlueprintCallable, Category = "Day Sequence")
+	void AddStaticTime(float StaticTime);
+	
+private:
+	FGuid GetOrCreateProceduralBinding(UObject* Object) const;
+
+	template<typename TrackType>
+	TrackType* CreateOrAddOverrideTrack(FName Name);
+
+	template<typename TrackType>
+	TrackType* CreateOrAddPropertyOverrideTrack(FName Name);
+
+	template<typename TrackType, typename SectionType>
+	SectionType* CreateOrAddPropertyOverrideSection(FName Name);
+	
+private:
+	/** This is returned immediately upon creation in InitializeSequence. The caller is responsible for holding a reference to prevent GC. */
+	TObjectPtr<UDaySequence> ProceduralDaySequence = nullptr;
+
+	TObjectPtr<ADaySequenceActor> TargetActor = nullptr;
+	
+	TObjectPtr<UObject> ActiveBoundObject = nullptr;
+	FGuid ActiveBinding;
+};
\ No newline at end of file
diff --git a/Engine/Plugins/Experimental/DaySequence/Source/DaySequenceEditor/Private/AssetDefinition_DaySequenceCollectionAsset.h b/Engine/Plugins/Experimental/DaySequence/Source/DaySequenceEditor/Private/AssetDefinition_DaySequenceCollectionAsset.h
new file mode 100644
--- /dev/null
+++ b/Engine/Plugins/Experimental/DaySequence/Source/DaySequenceEditor/Private/AssetDefinition_DaySequenceCollectionAsset.h
@@ -0,0 +1,32 @@
+// Copyright Epic Games, Inc. All Rights Reserved.
+
+#pragma once
+
+#include "AssetDefinitionDefault.h"
+#include "Factories/Factory.h"
+
+#include "AssetDefinition_DaySequenceCollectionAsset.generated.h"
+
+class UDaySequenceCollectionAsset;
+
+UCLASS()
+class UAssetDefinition_DaySequenceCollectionAsset : public UAssetDefinitionDefault
+{
+	GENERATED_BODY()
+
+public:
+	virtual FText GetAssetDisplayName() const override { return NSLOCTEXT("AssetTypeActions", "AssetTypeActions_DaySequenceCollectionAsset", "Day Sequence Collection"); }
+	virtual FLinearColor GetAssetColor() const override { return FLinearColor(FColor(200, 80, 80)); }
+	virtual TSoftClassPtr<UObject> GetAssetClass() const override;
+};
+
+UCLASS()
+class DAYSEQUENCEEDITOR_API UDaySequenceCollectionAssetFactory : public UFactory
+{
+	GENERATED_BODY()
+	
+public:
+	UDaySequenceCollectionAssetFactory(const class FObjectInitializer& Obj);
+	
+	virtual UObject* FactoryCreateNew(UClass* Class, UObject* InParent, FName Name, EObjectFlags Flags, UObject* Context, FFeedbackContext* Warn) override;
+};
\ No newline at end of file
diff --git a/Engine/Plugins/Experimental/DaySequence/Source/DaySequenceEditor/Private/DaySequenceEditorSubsystem.cpp b/Engine/Plugins/Experimental/DaySequence/Source/DaySequenceEditor/Private/DaySequenceEditorSubsystem.cpp
new file mode 100644
--- /dev/null
+++ b/Engine/Plugins/Experimental/DaySequence/Source/DaySequenceEditor/Private/DaySequenceEditorSubsystem.cpp
@@ -0,0 +1,1548 @@
+// Copyright Epic Games, Inc. All Rights Reserved.
+
+#include "DaySequenceEditorSubsystem.h"
+
+#include "Evaluation/MovieScenePlayback.h"
+#include "ISequencer.h"
+#include "Framework/Commands/UICommandList.h"
+#include "ISequencerModule.h"
+#include "Framework/MultiBox/SToolBarButtonBlock.h"
+#include "IDaySequenceEditorModule.h"
+#include "DaySequence.h"
+#include "DaySequenceActorPreview.h"
+#include "DaySequenceEditorCommands.h"
+#include "MovieScene.h"
+#include "MovieSceneCommonHelpers.h"
+#include "MovieScenePossessable.h"
+#include "MovieSceneSection.h"
+#include "MovieSceneSequence.h"
+#include "MovieSceneSpawnable.h"
+#include "MovieSceneTimeHelpers.h"
+#include "Selection.h"
+#include "SequencerSettings.h"
+#include "SequencerUtilities.h"
+#include "MovieSceneBindingProxy.h"
+#include "MovieSceneTimeUnit.h"
+#include "Sections/MovieScene3DConstraintSection.h"
+#include "Sections/MovieScene3DTransformSection.h"
+#include "Toolkits/AssetEditorToolkit.h"
+#include "Tracks/MovieScene3DTransformTrack.h"
+#include "Tracks/MovieScene3DConstraintTrack.h"
+#include "Tracks/MovieSceneCameraShakeTrack.h"
+
+#include "ActorTreeItem.h"
+#include "Camera/CameraComponent.h"
+#include "Components/SkeletalMeshComponent.h"
+#include "EngineUtils.h"
+#include "Framework/Application/SlateApplication.h"
+#include "Framework/MultiBox/MultiBoxBuilder.h"
+#include "GameFramework/Actor.h"
+#include "HAL/PlatformApplicationMisc.h"
+#include "Modules/ModuleManager.h"
+#include "SceneOutlinerModule.h"
+#include "SceneOutlinerPublicTypes.h"
+#include "ScopedTransaction.h"
+#include "Widgets/Notifications/SNotificationList.h"
+#include "Framework/MultiBox/MultiBoxBuilder.h"
+
+#include UE_INLINE_GENERATED_CPP_BY_NAME(DaySequenceEditorSubsystem)
+
+DEFINE_LOG_CATEGORY(LogDaySequenceEditor);
+
+#define LOCTEXT_NAMESPACE "DaySequenceEditor"
+
+void UDaySequenceEditorSubsystem::Initialize(FSubsystemCollectionBase& Collection)
+{
+	UE_LOG(LogDaySequenceEditor, Log, TEXT("DaySequenceEditor subsystem initialized."));
+
+	ISequencerModule& SequencerModule = FModuleManager::Get().LoadModuleChecked<ISequencerModule>("Sequencer");
+	OnSequencerCreatedHandle = SequencerModule.RegisterOnSequencerCreated(FOnSequencerCreated::FDelegate::CreateUObject(this, &UDaySequenceEditorSubsystem::OnSequencerCreated));
+
+	auto AreActorsSelected = [this]{
+		TArray<AActor*> SelectedActors;
+		GEditor->GetSelectedActors()->GetSelectedObjects<AActor>(SelectedActors); 
+		return SelectedActors.Num() > 0;
+	};
+
+	auto AreMovieSceneSectionsSelected = [this](const int32 MinSections = 1) {
+		const TSharedPtr<ISequencer> Sequencer = GetActiveSequencer();
+		if (!Sequencer)
+		{
+			return false;
+		}
+
+		TArray<UMovieSceneSection*> SelectedSections;
+		Sequencer->GetSelectedSections(SelectedSections);
+		return (SelectedSections.Num() >= MinSections);
+	};
+
+	/* Commands for this subsystem */
+	CommandList = MakeShareable(new FUICommandList);
+
+	CommandList->MapAction(FDaySequenceEditorCommands::Get().SnapSectionsToTimelineUsingSourceTimecode,
+		FExecuteAction::CreateUObject(this, &UDaySequenceEditorSubsystem::SnapSectionsToTimelineUsingSourceTimecodeInternal),
+		FCanExecuteAction::CreateLambda(AreMovieSceneSectionsSelected)
+	);
+
+	CommandList->MapAction(FDaySequenceEditorCommands::Get().SyncSectionsUsingSourceTimecode,
+		FExecuteAction::CreateUObject(this, &UDaySequenceEditorSubsystem::SyncSectionsUsingSourceTimecodeInternal),
+		FCanExecuteAction::CreateLambda(AreMovieSceneSectionsSelected, 2)
+	);
+
+	CommandList->MapAction(FDaySequenceEditorCommands::Get().BakeTransform,
+		FExecuteAction::CreateUObject(this, &UDaySequenceEditorSubsystem::BakeTransformInternal)
+	);
+
+	CommandList->MapAction(FDaySequenceEditorCommands::Get().FixActorReferences,
+		FExecuteAction::CreateUObject(this, &UDaySequenceEditorSubsystem::FixActorReferences)
+	);
+
+	CommandList->MapAction(FDaySequenceEditorCommands::Get().AddActorsToBinding,
+		FExecuteAction::CreateUObject(this, &UDaySequenceEditorSubsystem::AddActorsToBindingInternal),
+		FCanExecuteAction::CreateLambda(AreActorsSelected));
+
+	CommandList->MapAction(FDaySequenceEditorCommands::Get().ReplaceBindingWithActors,
+		FExecuteAction::CreateUObject(this, &UDaySequenceEditorSubsystem::ReplaceBindingWithActorsInternal),
+		FCanExecuteAction::CreateLambda(AreActorsSelected));
+
+	CommandList->MapAction(FDaySequenceEditorCommands::Get().RemoveActorsFromBinding,
+		FExecuteAction::CreateUObject(this, &UDaySequenceEditorSubsystem::RemoveActorsFromBindingInternal),
+		FCanExecuteAction::CreateLambda(AreActorsSelected));
+
+	CommandList->MapAction(FDaySequenceEditorCommands::Get().RemoveAllBindings,
+		FExecuteAction::CreateUObject(this, &UDaySequenceEditorSubsystem::RemoveAllBindingsInternal));
+
+	CommandList->MapAction(FDaySequenceEditorCommands::Get().RemoveInvalidBindings,
+		FExecuteAction::CreateUObject(this, &UDaySequenceEditorSubsystem::RemoveInvalidBindingsInternal));
+
+	/* Menu extenders */
+	TransformMenuExtender = MakeShareable(new FExtender);
+	TransformMenuExtender->AddMenuExtension("Transform", EExtensionHook::After, CommandList, FMenuExtensionDelegate::CreateLambda([this](FMenuBuilder& MenuBuilder) {
+		// Only add menu entries where the focused sequence is a UDaySequence
+		if (!GetActiveSequencer())
+		{
+			return;
+		}
+		
+		MenuBuilder.AddMenuEntry(FDaySequenceEditorCommands::Get().SnapSectionsToTimelineUsingSourceTimecode);
+		MenuBuilder.AddMenuEntry(FDaySequenceEditorCommands::Get().SyncSectionsUsingSourceTimecode);
+		MenuBuilder.AddMenuEntry(FDaySequenceEditorCommands::Get().BakeTransform);
+		}));
+
+	SequencerModule.GetActionsMenuExtensibilityManager()->AddExtender(TransformMenuExtender);
+
+	FixActorReferencesMenuExtender = MakeShareable(new FExtender);
+	FixActorReferencesMenuExtender->AddMenuExtension("Bindings", EExtensionHook::First, CommandList, FMenuExtensionDelegate::CreateLambda([this](FMenuBuilder& MenuBuilder) {
+		// Only add menu entries where the focused sequence is a UDaySequence
+		if (!GetActiveSequencer())
+		{
+			return;
+		}
+		
+		MenuBuilder.AddMenuEntry(FDaySequenceEditorCommands::Get().FixActorReferences);
+		}));
+
+	SequencerModule.GetActionsMenuExtensibilityManager()->AddExtender(FixActorReferencesMenuExtender);
+
+	AssignActorMenuExtender = MakeShareable(new FExtender);
+	AssignActorMenuExtender->AddMenuExtension("Possessable", EExtensionHook::First, CommandList, FMenuExtensionDelegate::CreateLambda([this](FMenuBuilder& MenuBuilder) {
+		// Only add menu entries where the focused sequence is a UDaySequence
+		if (!GetActiveSequencer())
+		{
+			return;
+		}
+		
+		FFormatNamedArguments Args;
+		MenuBuilder.AddSubMenu(
+			FText::Format(LOCTEXT("AssignActor", "Assign Actor"), Args),
+			FText::Format(LOCTEXT("AssignActorTooltip", "Assign an actor to this track"), Args),
+			FNewMenuDelegate::CreateLambda([this](FMenuBuilder& SubMenuBuilder) { AddAssignActorMenu(SubMenuBuilder); } ));
+		}));
+
+	SequencerModule.GetObjectBindingContextMenuExtensibilityManager()->AddExtender(AssignActorMenuExtender);
+
+	RebindComponentMenuExtender = MakeShareable(new FExtender);
+	RebindComponentMenuExtender->AddMenuExtension("Possessable", EExtensionHook::First, CommandList, FMenuExtensionDelegate::CreateLambda([this](FMenuBuilder& MenuBuilder) {
+		// Only add menu entries where the focused sequence is a UDaySequence
+		if (!GetActiveSequencer())
+		{
+			return;
+		}
+		
+		TArray<FName> ComponentNames;
+		GetRebindComponentNames(ComponentNames);
+		if (ComponentNames.Num() > 0)
+		{
+			FFormatNamedArguments Args;
+			MenuBuilder.AddSubMenu(
+				FText::Format(LOCTEXT("RebindComponent", "Rebind Component"), Args),
+				FText::Format(LOCTEXT("RebindComponentTooltip", "Rebind component by moving the tracks from one component to another component."), Args),
+				FNewMenuDelegate::CreateLambda([this](FMenuBuilder& SubMenuBuilder) { RebindComponentMenu(SubMenuBuilder); } ));
+		}
+	}));
+
+	SequencerModule.GetObjectBindingContextMenuExtensibilityManager()->AddExtender(RebindComponentMenuExtender);
+}
+
+void UDaySequenceEditorSubsystem::Deinitialize()
+{
+	UE_LOG(LogDaySequenceEditor, Log, TEXT("DaySequenceEditor subsystem deinitialized."));
+
+	ISequencerModule* SequencerModulePtr = FModuleManager::Get().GetModulePtr<ISequencerModule>("Sequencer");
+	if (SequencerModulePtr)
+	{
+		SequencerModulePtr->UnregisterOnSequencerCreated(OnSequencerCreatedHandle);
+	}
+
+}
+
+void UDaySequenceEditorSubsystem::OnSequencerCreated(TSharedRef<ISequencer> InSequencer)
+{
+	UE_LOG(LogDaySequenceEditor, VeryVerbose, TEXT("UDaySequenceEditorSubsystem::OnSequencerCreated"));
+
+	Sequencers.Add(TWeakPtr<ISequencer>(InSequencer));
+}
+
+TSharedPtr<ISequencer> UDaySequenceEditorSubsystem::GetActiveSequencer()
+{
+	// Disregard preview day sequence since it is mutually exclusive with the day sequence editor
+	UMovieSceneSequence* PreviewSequence = nullptr;
+	IDaySequenceEditorModule* DaySequenceEditorModule = FModuleManager::GetModulePtr<IDaySequenceEditorModule>("DaySequenceEditor");
+	if (DaySequenceEditorModule)
+	{
+		FDaySequenceActorPreview& DaySequenceActorPreview = DaySequenceEditorModule->GetDaySequenceActorPreview();
+		TWeakPtr<ISequencer> PreviewSequencer = DaySequenceActorPreview.GetPreviewSequencer();
+		if (PreviewSequencer.IsValid())
+		{
+			PreviewSequence = PreviewSequencer.Pin()->GetFocusedMovieSceneSequence();
+		}
+	}
+
+	for (TWeakPtr<ISequencer> Ptr : Sequencers)
+	{
+		if (Ptr.IsValid())
+		{
+			UMovieSceneSequence* Sequence = Ptr.Pin()->GetFocusedMovieSceneSequence();
+			if (Sequence && Sequence->IsA<UDaySequence>() && PreviewSequence != Sequence)
+			{
+				return Ptr.Pin();
+			}
+		}
+	}
+
+	return nullptr;
+}
+
+TArray<FMovieSceneBindingProxy> UDaySequenceEditorSubsystem::AddActors(const TArray<AActor*>& InActors)
+{
+	TArray<FMovieSceneBindingProxy> BindingProxies;
+
+	TSharedPtr<ISequencer> Sequencer = GetActiveSequencer();
+	if (Sequencer == nullptr)
+	{
+		return BindingProxies;
+	}
+
+	UMovieSceneSequence* Sequence = Sequencer->GetFocusedMovieSceneSequence();
+	if (!Sequence)
+	{
+		return BindingProxies;
+	}
+
+	TArray<TWeakObjectPtr<AActor> > Actors;
+	for (AActor* Actor : InActors)
+	{
+		Actors.Add(Actor);
+	}
+
+	TArray<FGuid> Guids = FSequencerUtilities::AddActors(Sequencer.ToSharedRef(), Actors);
+	
+	for (const FGuid& Guid : Guids)
+	{
+		BindingProxies.Add(FMovieSceneBindingProxy(Guid, Sequence));
+	}
+
+	return BindingProxies;
+}
+
+TArray<FMovieSceneBindingProxy> UDaySequenceEditorSubsystem::ConvertToSpawnable(const FMovieSceneBindingProxy& ObjectBinding)
+{
+	TArray<FMovieSceneBindingProxy> SpawnableProxies;
+
+	TSharedPtr<ISequencer> Sequencer = GetActiveSequencer();
+	if (Sequencer == nullptr)
+	{
+		return SpawnableProxies;
+	}
+
+	UMovieSceneSequence* Sequence = Sequencer->GetFocusedMovieSceneSequence();
+	if (!Sequence)
+	{
+		return SpawnableProxies;
+	}
+
+	TArray<FMovieSceneSpawnable*> Spawnables = FSequencerUtilities::ConvertToSpawnable(Sequencer.ToSharedRef(), ObjectBinding.BindingID);
+	for (FMovieSceneSpawnable* Spawnable : Spawnables)
+	{
+		if (Spawnable)
+		{
+			SpawnableProxies.Add(FMovieSceneBindingProxy(Spawnable->GetGuid(), Sequence));
+		}
+	}
+
+	return SpawnableProxies;
+}
+
+FMovieSceneBindingProxy UDaySequenceEditorSubsystem::ConvertToPossessable(const FMovieSceneBindingProxy& ObjectBinding)
+{
+	FMovieSceneBindingProxy PossessableProxy;
+
+	TSharedPtr<ISequencer> Sequencer = GetActiveSequencer();
+	if (Sequencer == nullptr)
+	{
+		return PossessableProxy;
+	}
+
+	UMovieSceneSequence* Sequence = Sequencer->GetFocusedMovieSceneSequence();
+	if (!Sequence)
+	{
+		return PossessableProxy;
+	}
+
+	if (FMovieScenePossessable* Possessable = FSequencerUtilities::ConvertToPossessable(Sequencer.ToSharedRef(), ObjectBinding.BindingID))
+	{
+		PossessableProxy = FMovieSceneBindingProxy(Possessable->GetGuid(), Sequence);
+	}
+
+	return PossessableProxy;
+}
+
+void UDaySequenceEditorSubsystem::CopyFolders(const TArray<UMovieSceneFolder*>& Folders, FString& ExportedText)
+{
+	FSequencerUtilities::CopyFolders(Folders, ExportedText);
+
+	FPlatformApplicationMisc::ClipboardCopy(*ExportedText);
+}
+
+bool UDaySequenceEditorSubsystem::PasteFolders(const FString& InTextToImport, FMovieScenePasteFoldersParams PasteFoldersParams, TArray<UMovieSceneFolder*>& OutFolders)
+{
+	FString TextToImport = InTextToImport;
+	if (TextToImport.IsEmpty())
+	{
+		FPlatformApplicationMisc::ClipboardPaste(TextToImport);
+	}
+
+	TArray<FNotificationInfo> PasteErrors;
+	if (!FSequencerUtilities::PasteFolders(TextToImport, PasteFoldersParams, OutFolders, PasteErrors))
+	{
+		for (FNotificationInfo PasteError : PasteErrors)
+		{
+			UE_LOG(LogDaySequenceEditor, Error, TEXT("%s"), *PasteError.Text.Get().ToString());
+		}
+		return false;
+	}
+
+	return true;
+}
+
+void UDaySequenceEditorSubsystem::CopySections(const TArray<UMovieSceneSection*>& Sections, FString& ExportedText)
+{
+	FSequencerUtilities::CopySections(Sections, ExportedText);
+
+	FPlatformApplicationMisc::ClipboardCopy(*ExportedText);
+}
+
+bool UDaySequenceEditorSubsystem::PasteSections(const FString& InTextToImport, FMovieScenePasteSectionsParams PasteSectionsParams, TArray<UMovieSceneSection*>& OutSections)
+{
+	FString TextToImport = InTextToImport;
+	if (TextToImport.IsEmpty())
+	{
+		FPlatformApplicationMisc::ClipboardPaste(TextToImport);
+	}
+
+	TArray<FNotificationInfo> PasteErrors;
+	if (!FSequencerUtilities::PasteSections(TextToImport, PasteSectionsParams, OutSections, PasteErrors))
+	{
+		for (FNotificationInfo PasteError : PasteErrors)
+		{
+			UE_LOG(LogDaySequenceEditor, Error, TEXT("%s"), *PasteError.Text.Get().ToString());
+		}
+		return false;
+	}
+
+	return true;
+}
+
+void UDaySequenceEditorSubsystem::CopyTracks(const TArray<UMovieSceneTrack*>& Tracks, FString& ExportedText)
+{
+	TArray<UMovieSceneFolder*> Folders;
+	FSequencerUtilities::CopyTracks(Tracks, Folders, ExportedText);
+	FPlatformApplicationMisc::ClipboardCopy(*ExportedText);
+}
+
+bool UDaySequenceEditorSubsystem::PasteTracks(const FString& InTextToImport, FMovieScenePasteTracksParams PasteTracksParams, TArray<UMovieSceneTrack*>& OutTracks)
+{
+	FString TextToImport = InTextToImport;
+	if (TextToImport.IsEmpty())
+	{
+		FPlatformApplicationMisc::ClipboardPaste(TextToImport);
+	}
+
+	TArray<FNotificationInfo> PasteErrors;
+	if (!FSequencerUtilities::PasteTracks(TextToImport, PasteTracksParams, OutTracks, PasteErrors))
+	{
+		for (FNotificationInfo PasteError : PasteErrors)
+		{
+			UE_LOG(LogDaySequenceEditor, Error, TEXT("%s"), *PasteError.Text.Get().ToString());
+		}
+		return false;
+	}
+
+	return true;
+}
+
+void UDaySequenceEditorSubsystem::CopyBindings(const TArray<FMovieSceneBindingProxy>& Bindings, FString& ExportedText)
+{
+	TSharedPtr<ISequencer> Sequencer = GetActiveSequencer();
+	if (Sequencer == nullptr)
+	{
+		return;
+	}
+
+	TArray<UMovieSceneFolder*> Folders;
+	FSequencerUtilities::CopyBindings(Sequencer.ToSharedRef(), Bindings, Folders, ExportedText);
+	FPlatformApplicationMisc::ClipboardCopy(*ExportedText);
+}
+
+bool UDaySequenceEditorSubsystem::PasteBindings(const FString& InTextToImport, FMovieScenePasteBindingsParams PasteBindingsParams, TArray<FMovieSceneBindingProxy>& OutObjectBindings)
+{
+	TSharedPtr<ISequencer> Sequencer = GetActiveSequencer();
+	if (Sequencer == nullptr)
+	{
+		return false;
+	}
+
+	FString TextToImport = InTextToImport;
+	if (TextToImport.IsEmpty())
+	{
+		FPlatformApplicationMisc::ClipboardPaste(TextToImport);
+	}
+
+	TArray<FNotificationInfo> PasteErrors;
+	TArray<FMovieSceneBindingProxy> OutBindings;
+	if (!FSequencerUtilities::PasteBindings(TextToImport, Sequencer.ToSharedRef(), PasteBindingsParams, OutBindings, PasteErrors))
+	{
+		for (FNotificationInfo PasteError : PasteErrors)
+		{
+			UE_LOG(LogDaySequenceEditor, Error, TEXT("%s"), *PasteError.Text.Get().ToString());
+		}
+		return false;
+	}
+
+	return true;
+}
+
+void UDaySequenceEditorSubsystem::SnapSectionsToTimelineUsingSourceTimecodeInternal()
+{
+	const TSharedPtr<ISequencer> Sequencer = GetActiveSequencer();
+	if (!Sequencer)
+	{
+		return;
+	}
+
+	const UMovieScene* FocusedMovieScene = Sequencer->GetFocusedMovieSceneSequence()->GetMovieScene();
+	if (!FocusedMovieScene)
+	{
+		return;
+	}
+
+	if (FocusedMovieScene->IsReadOnly())
+	{
+		FSequencerUtilities::ShowReadOnlyError();
+		return;
+	}
+
+	TArray<UMovieSceneSection*> Sections;
+	Sequencer->GetSelectedSections(Sections);
+	if (Sections.IsEmpty())
+	{
+		return;
+	}
+
+	SnapSectionsToTimelineUsingSourceTimecode(Sections);
+}
+
+void UDaySequenceEditorSubsystem::SnapSectionsToTimelineUsingSourceTimecode(const TArray<UMovieSceneSection*>& Sections)
+{
+	const TSharedPtr<ISequencer> Sequencer = GetActiveSequencer();
+	if (!Sequencer)
+	{
+		return;
+	}
+
+	const UMovieScene* FocusedMovieScene = Sequencer->GetFocusedMovieSceneSequence()->GetMovieScene();
+	if (!FocusedMovieScene)
+	{
+		return;
+	}
+
+	if (FocusedMovieScene->IsReadOnly())
+	{
+		FSequencerUtilities::ShowReadOnlyError();
+		return;
+	}
+
+	FScopedTransaction SnapSectionsToTimelineUsingSourceTimecodeTransaction(LOCTEXT("SnapSectionsToTimelineUsingSourceTimecode_Transaction", "Snap Sections to Timeline using Source Timecode"));
+	bool bAnythingChanged = false;
+
+	const FFrameRate TickResolution = Sequencer->GetFocusedTickResolution();
+
+	for (UMovieSceneSection* Section : Sections)
+	{
+		if (!Section || !Section->HasStartFrame())
+		{
+			continue;
+		}
+
+		const FTimecode SectionSourceTimecode = Section->TimecodeSource.Timecode;
+		if (SectionSourceTimecode == FTimecode())
+		{
+			// Do not move sections with default values for source timecode.
+			continue;
+		}
+
+		const FFrameNumber SectionSourceStartFrameNumber = SectionSourceTimecode.ToFrameNumber(TickResolution);
+
+		// Account for any trimming at the start of the section when computing the
+		// target frame number to move this section to.
+		const FFrameNumber SectionOffsetFrames = Section->GetOffsetTime().Get(FFrameTime()).FloorToFrame();
+		const FFrameNumber TargetFrameNumber = SectionSourceStartFrameNumber + SectionOffsetFrames;
+
+		const FFrameNumber SectionCurrentStartFrameNumber = Section->GetInclusiveStartFrame();
+
+		const FFrameNumber Delta = -(SectionCurrentStartFrameNumber - TargetFrameNumber);
+
+		Section->MoveSection(Delta);
+
+		bAnythingChanged |= (Delta.Value != 0);
+	}
+
+	if (bAnythingChanged)
+	{
+		Sequencer->NotifyMovieSceneDataChanged(EMovieSceneDataChangeType::TrackValueChanged);
+	}
+}
+
+void UDaySequenceEditorSubsystem::SyncSectionsUsingSourceTimecodeInternal()
+{
+	const TSharedPtr<ISequencer> Sequencer = GetActiveSequencer();
+	if (!Sequencer)
+	{
+		return;
+	}
+
+	const UMovieScene* FocusedMovieScene = Sequencer->GetFocusedMovieSceneSequence()->GetMovieScene();
+	if (!FocusedMovieScene)
+	{
+		return;
+	}
+
+	if (FocusedMovieScene->IsReadOnly())
+	{
+		FSequencerUtilities::ShowReadOnlyError();
+		return;
+	}
+
+	TArray<UMovieSceneSection*> Sections;
+	Sequencer->GetSelectedSections(Sections);
+	if (Sections.Num() < 2)
+	{
+		return;
+	}
+
+	SyncSectionsUsingSourceTimecode(Sections);
+}
+
+void UDaySequenceEditorSubsystem::SyncSectionsUsingSourceTimecode(const TArray<UMovieSceneSection*>& Sections)
+{
+	const TSharedPtr<ISequencer> Sequencer = GetActiveSequencer();
+	if (!Sequencer)
+	{
+		return;
+	}
+
+	const UMovieScene* FocusedMovieScene = Sequencer->GetFocusedMovieSceneSequence()->GetMovieScene();
+	if (!FocusedMovieScene)
+	{
+		return;
+	}
+
+	if (FocusedMovieScene->IsReadOnly())
+	{
+		FSequencerUtilities::ShowReadOnlyError();
+		return;
+	}
+
+	// Pull out all of the valid sections that have a start frame and verify
+	// we have at least two sections to sync.
+	TArray<UMovieSceneSection*> SectionsToSync;
+	for (UMovieSceneSection* Section : Sections)
+	{
+		if (Section && Section->HasStartFrame())
+		{
+			SectionsToSync.Add(Section);
+		}
+	}
+
+	if (SectionsToSync.Num() < 2)
+	{
+		return;
+	}
+
+	FScopedTransaction SyncSectionsUsingSourceTimecodeTransaction(LOCTEXT("SyncSectionsUsingSourceTimecode_Transaction", "Sync Sections Using Source Timecode"));
+	bool bAnythingChanged = false;
+
+	const FFrameRate TickResolution = Sequencer->GetFocusedTickResolution();
+
+	const UMovieSceneSection* FirstSection = SectionsToSync[0];
+	const FFrameNumber FirstSectionSourceTimecode = FirstSection->TimecodeSource.Timecode.ToFrameNumber(TickResolution);
+	const FFrameNumber FirstSectionCurrentStartFrame = FirstSection->GetInclusiveStartFrame();
+	const FFrameNumber FirstSectionOffsetFrames = FirstSection->GetOffsetTime().Get(FFrameTime()).FloorToFrame();
+	SectionsToSync.RemoveAt(0);
+
+	for (UMovieSceneSection* Section : SectionsToSync)
+	{
+		const FFrameNumber SectionSourceTimecode = Section->TimecodeSource.Timecode.ToFrameNumber(TickResolution);
+		const FFrameNumber SectionCurrentStartFrame = Section->GetInclusiveStartFrame();
+		const FFrameNumber SectionOffsetFrames = Section->GetOffsetTime().Get(FFrameTime()).FloorToFrame();
+
+		const FFrameNumber TimecodeDelta = SectionSourceTimecode - FirstSectionSourceTimecode;
+		const FFrameNumber CurrentDelta = (SectionCurrentStartFrame - SectionOffsetFrames) - (FirstSectionCurrentStartFrame - FirstSectionOffsetFrames);
+		const FFrameNumber Delta = -CurrentDelta + TimecodeDelta;
+
+		Section->MoveSection(Delta);
+
+		bAnythingChanged |= (Delta.Value != 0);
+	}
+
+	if (bAnythingChanged)
+	{
+		Sequencer->NotifyMovieSceneDataChanged(EMovieSceneDataChangeType::TrackValueChanged);
+	}
+}
+
+void UDaySequenceEditorSubsystem::BakeTransformInternal()
+{
+	TSharedPtr<ISequencer> Sequencer = GetActiveSequencer();
+	if (Sequencer == nullptr)
+	{
+		return;
+	}
+
+	UMovieScene* FocusedMovieScene = Sequencer->GetFocusedMovieSceneSequence()->GetMovieScene();
+	if (!FocusedMovieScene)
+	{
+		return;
+	}
+
+	FFrameRate TickResolution = FocusedMovieScene->GetTickResolution();
+	FFrameRate DisplayRate = FocusedMovieScene->GetDisplayRate();
+
+	FFrameNumber InFrame = UE::MovieScene::DiscreteInclusiveLower(FocusedMovieScene->GetPlaybackRange());
+	FFrameNumber OutFrame = UE::MovieScene::DiscreteExclusiveUpper(FocusedMovieScene->GetPlaybackRange());
+
+	TArray<FGuid> ObjectBindings;
+	Sequencer->GetSelectedObjects(ObjectBindings);
+
+	TArray<FMovieSceneBindingProxy> BindingProxies;
+	for (FGuid Guid : ObjectBindings)
+	{
+		BindingProxies.Add(FMovieSceneBindingProxy(Guid, Sequencer->GetFocusedMovieSceneSequence()));
+	}
+
+	FMovieSceneScriptingParams Params;
+	Params.TimeUnit = EMovieSceneTimeUnit::TickResolution;
+	BakeTransform(BindingProxies, FFrameTime(InFrame), FFrameTime(OutFrame), ConvertFrameTime(1, DisplayRate, TickResolution), Params);
+}
+
+void UDaySequenceEditorSubsystem::BakeTransform(const TArray<FMovieSceneBindingProxy>& ObjectBindings, const FFrameTime& BakeInTime, const FFrameTime& BakeOutTime, const FFrameTime& BakeInterval, const FMovieSceneScriptingParams& Params)
+{
+	TSharedPtr<ISequencer> Sequencer = GetActiveSequencer();
+	if (Sequencer == nullptr)
+	{
+		return;
+	}
+
+	UMovieScene* FocusedMovieScene = Sequencer->GetFocusedMovieSceneSequence()->GetMovieScene();
+	if (!FocusedMovieScene)
+	{
+		return;
+	}
+
+	if (FocusedMovieScene->IsReadOnly())
+	{
+		FSequencerUtilities::ShowReadOnlyError();
+		return;
+	}
+
+	if (ObjectBindings.Num() == 0)
+	{
+		return;
+	}
+
+	FScopedTransaction BakeTransform(LOCTEXT("BakeTransform", "Bake Transform"));
+
+	FocusedMovieScene->Modify();
+
+	FQualifiedFrameTime ResetTime = Sequencer->GetLocalTime();
+
+	FFrameRate TickResolution = FocusedMovieScene->GetTickResolution();
+	FFrameRate DisplayRate = FocusedMovieScene->GetDisplayRate();
+
+	FFrameTime InFrame = BakeInTime;
+	FFrameTime OutFrame = BakeOutTime;
+	FFrameTime Interval = BakeInterval;
+
+	if (Params.TimeUnit == EMovieSceneTimeUnit::DisplayRate)
+	{
+		InFrame = ConvertFrameTime(BakeInTime, DisplayRate, TickResolution);
+		OutFrame = ConvertFrameTime(BakeOutTime, DisplayRate, TickResolution);
+		Interval = ConvertFrameTime(BakeInterval, DisplayRate, TickResolution);
+	}
+
+	struct FBakeData
+	{
+		TArray<FVector> Locations;
+		TArray<FRotator> Rotations;
+		TArray<FVector> Scales;
+		TArray<FFrameNumber> KeyTimes;
+	};
+
+	TMap<FGuid, FBakeData> BakeDataMap;
+	for (const FMovieSceneBindingProxy& ObjectBinding : ObjectBindings)
+	{
+		BakeDataMap.Add(ObjectBinding.BindingID);
+	}
+
+	FMovieSceneSequenceTransform RootToLocalTransform = Sequencer->GetFocusedMovieSceneSequenceTransform();
+	
+	for (FFrameTime EvalTime = InFrame; EvalTime <= OutFrame; EvalTime += Interval)
+	{
+		FFrameNumber KeyTime = FFrameRate::Snap(EvalTime, TickResolution, DisplayRate).FloorToFrame();
+		FMovieSceneEvaluationRange Range(KeyTime * RootToLocalTransform.InverseNoLooping(), TickResolution);
+
+		Sequencer->SetLocalTimeDirectly(Range.GetTime());
+
+		for (const FMovieSceneBindingProxy& ObjectBinding : ObjectBindings)
+		{
+			FGuid Guid = ObjectBinding.BindingID;
+
+			for (TWeakObjectPtr<> RuntimeObject : Sequencer->FindObjectsInCurrentSequence(Guid) )
+			{
+				AActor* Actor = Cast<AActor>(RuntimeObject.Get());
+				if (!Actor)
+				{
+					UActorComponent* ActorComponent = Cast<UActorComponent>(RuntimeObject.Get());
+					if (ActorComponent)
+					{
+						Actor = ActorComponent->GetOwner();
+					}
+				}
+
+				if (!Actor)
+				{
+					continue;
+				}
+
+				UCameraComponent* CameraComponent = MovieSceneHelpers::CameraComponentFromRuntimeObject(RuntimeObject.Get());
+
+				// Cache transforms
+				USceneComponent* Parent = nullptr;
+				if (CameraComponent)
+				{
+					Parent = CameraComponent->GetAttachParent();
+				} 
+				else if (Actor->GetRootComponent())
+				{
+					Parent = Actor->GetRootComponent()->GetAttachParent();
+				}
+				
+				// The CameraRig_rail updates the spline position tick, so it needs to be ticked manually while baking the frames
+				while (Parent && Parent->GetOwner())
+				{
+					Parent->GetOwner()->Tick(0.03f);
+					if (USkeletalMeshComponent* SkeletalMeshComponent = Cast<USkeletalMeshComponent>(Parent))
+					{
+						SkeletalMeshComponent->TickAnimation(0.f, false);
+
+						SkeletalMeshComponent->RefreshBoneTransforms();
+						SkeletalMeshComponent->RefreshFollowerComponents();
+						SkeletalMeshComponent->UpdateComponentToWorld();
+						SkeletalMeshComponent->FinalizeBoneTransform();
+						SkeletalMeshComponent->MarkRenderTransformDirty();
+						SkeletalMeshComponent->MarkRenderDynamicDataDirty();
+					}
+					Parent = Parent->GetAttachParent();
+				}
+
+				if (CameraComponent)
+				{
+					FTransform AdditiveOffset;
+					float AdditiveFOVOffset;
+					CameraComponent->GetAdditiveOffset(AdditiveOffset, AdditiveFOVOffset);
+
+					FTransform Transform(Actor->GetActorRotation(), Actor->GetActorLocation());
+					FTransform TransformWithAdditiveOffset = AdditiveOffset * Transform;
+					FVector LocalTranslation = TransformWithAdditiveOffset.GetTranslation();
+					FRotator LocalRotation = TransformWithAdditiveOffset.GetRotation().Rotator();
+
+					BakeDataMap[Guid].Locations.Add(LocalTranslation);
+					BakeDataMap[Guid].Rotations.Add(LocalRotation);
+					BakeDataMap[Guid].Scales.Add(FVector::OneVector);
+				}
+				else
+				{
+					BakeDataMap[Guid].Locations.Add(Actor->GetActorLocation());
+					BakeDataMap[Guid].Rotations.Add(Actor->GetActorRotation());
+					BakeDataMap[Guid].Scales.Add(Actor->GetActorScale());
+				}
+
+				BakeDataMap[Guid].KeyTimes.Add(KeyTime);
+			}
+		}
+	}
+
+	const bool bDisableSectionsAfterBaking = Sequencer->GetSequencerSettings()->GetDisableSectionsAfterBaking();
+
+	for (TPair<FGuid, FBakeData>& BakeData : BakeDataMap)
+	{
+		FGuid Guid = BakeData.Key;
+
+		// Disable or delete any constraint (attach/path) tracks
+		AActor* AttachParentActor = nullptr;
+		for (UMovieSceneTrack* Track : FocusedMovieScene->FindTracks(UMovieScene3DConstraintTrack::StaticClass(), Guid))
+		{
+			if (UMovieScene3DConstraintTrack* ConstraintTrack = Cast<UMovieScene3DConstraintTrack>(Track))
+			{
+				for (UMovieSceneSection* ConstraintSection : ConstraintTrack->GetAllSections())
+				{
+					FMovieSceneObjectBindingID ConstraintBindingID = (Cast<UMovieScene3DConstraintSection>(ConstraintSection))->GetConstraintBindingID();
+					for (TWeakObjectPtr<> ParentObject : ConstraintBindingID.ResolveBoundObjects(Sequencer->GetFocusedTemplateID(), *Sequencer))
+					{
+						AttachParentActor = Cast<AActor>(ParentObject.Get());
+						break;
+					}
+				}
+
+				if (bDisableSectionsAfterBaking)
+				{
+					for (UMovieSceneSection* ConstraintSection : ConstraintTrack->GetAllSections())
+					{
+						ConstraintSection->Modify();
+						ConstraintSection->SetIsActive(false);
+					}
+				}
+				else
+				{
+					FocusedMovieScene->RemoveTrack(*ConstraintTrack);
+				}
+			}
+		}
+
+		// Disable or delete any transform tracks
+		for (UMovieSceneTrack* Track : FocusedMovieScene->FindTracks(UMovieScene3DTransformTrack::StaticClass(), Guid))
+		{
+			if (UMovieScene3DTransformTrack* TransformTrack = Cast<UMovieScene3DTransformTrack>(Track))
+			{
+				if (bDisableSectionsAfterBaking)
+				{
+					for (UMovieSceneSection* TransformSection : TransformTrack->GetAllSections())
+					{
+						TransformSection->Modify();
+						TransformSection->SetIsActive(false);
+					}
+				}
+				else
+				{
+					FocusedMovieScene->RemoveTrack(*TransformTrack);
+				}
+			}
+		}
+
+		// Disable or delete any camera shake tracks
+		for (UMovieSceneTrack* Track : FocusedMovieScene->FindTracks(UMovieSceneCameraShakeTrack::StaticClass(), Guid))
+		{
+			if (UMovieSceneCameraShakeTrack* CameraShakeTrack = Cast<UMovieSceneCameraShakeTrack>(Track))
+			{
+				if (bDisableSectionsAfterBaking)
+				{
+					for (UMovieSceneSection* CameraShakeSection : CameraShakeTrack->GetAllSections())
+					{
+						CameraShakeSection->Modify();
+						CameraShakeSection->SetIsActive(false);
+					}
+				}
+				else
+				{
+					FocusedMovieScene->RemoveTrack(*CameraShakeTrack);
+				}
+			}
+		}
+
+		// Reset position
+		Sequencer->SetLocalTimeDirectly(ResetTime.Time);
+		Sequencer->ForceEvaluate();
+
+		FVector DefaultLocation = FVector::ZeroVector;
+		FVector DefaultRotation = FVector::ZeroVector;
+		FVector DefaultScale = FVector::OneVector;
+
+		for (TWeakObjectPtr<> RuntimeObject : Sequencer->FindObjectsInCurrentSequence(Guid))
+		{
+			AActor* Actor = Cast<AActor>(RuntimeObject.Get());
+			if (!Actor)
+			{
+				UActorComponent* ActorComponent = Cast<UActorComponent>(RuntimeObject.Get());
+				if (ActorComponent)
+				{
+					Actor = ActorComponent->GetOwner();
+				}
+			}
+
+			if (!Actor)
+			{
+				continue;
+			}
+
+			DefaultLocation = Actor->GetActorLocation();
+			DefaultRotation = Actor->GetActorRotation().Euler();
+			DefaultScale = Actor->GetActorScale();
+
+			// Always detach from any existing parent
+			Actor->DetachFromActor(FDetachmentTransformRules::KeepRelativeTransform);
+		}
+			
+		// Create new transform track and section
+		UMovieScene3DTransformTrack* TransformTrack = Cast<UMovieScene3DTransformTrack>(FocusedMovieScene->AddTrack(UMovieScene3DTransformTrack::StaticClass(), Guid));
+
+		if (TransformTrack)
+		{
+			UMovieScene3DTransformSection* TransformSection = CastChecked<UMovieScene3DTransformSection>(TransformTrack->CreateNewSection());
+			TransformTrack->AddSection(*TransformSection);
+
+			TransformSection->SetRange(TRange<FFrameNumber>::All());
+
+			TArrayView<FMovieSceneDoubleChannel*> DoubleChannels = TransformSection->GetChannelProxy().GetChannels<FMovieSceneDoubleChannel>();
+			DoubleChannels[0]->SetDefault(DefaultLocation.X);
+			DoubleChannels[1]->SetDefault(DefaultLocation.Y);
+			DoubleChannels[2]->SetDefault(DefaultLocation.Z);
+			DoubleChannels[3]->SetDefault(DefaultRotation.X);
+			DoubleChannels[4]->SetDefault(DefaultRotation.Y);
+			DoubleChannels[5]->SetDefault(DefaultRotation.Z);
+			DoubleChannels[6]->SetDefault(DefaultScale.X);
+			DoubleChannels[7]->SetDefault(DefaultScale.Y);
+			DoubleChannels[8]->SetDefault(DefaultScale.Z);
+
+			TArray<FVector> LocalTranslations, LocalRotations, LocalScales;
+			LocalTranslations.SetNum(BakeData.Value.KeyTimes.Num());
+			LocalRotations.SetNum(BakeData.Value.KeyTimes.Num());
+			LocalScales.SetNum(BakeData.Value.KeyTimes.Num());
+
+			for (int32 Counter = 0; Counter < BakeData.Value.KeyTimes.Num(); ++Counter)
+			{
+				FTransform LocalTransform(BakeData.Value.Rotations[Counter], BakeData.Value.Locations[Counter], BakeData.Value.Scales[Counter]);
+				LocalTranslations[Counter] = LocalTransform.GetTranslation();
+				LocalRotations[Counter] = LocalTransform.GetRotation().Euler();
+				LocalScales[Counter] = LocalTransform.GetScale3D();
+			}
+
+			// Euler filter
+			for (int32 Counter = 0; Counter < LocalRotations.Num() - 1; ++Counter)
+			{
+				FMath::WindRelativeAnglesDegrees(LocalRotations[Counter].X, LocalRotations[Counter + 1].X);
+				FMath::WindRelativeAnglesDegrees(LocalRotations[Counter].Y, LocalRotations[Counter + 1].Y);
+				FMath::WindRelativeAnglesDegrees(LocalRotations[Counter].Z, LocalRotations[Counter + 1].Z);							
+			}
+				
+			for (int32 Counter = 0; Counter < BakeData.Value.KeyTimes.Num(); ++Counter)
+			{
+				FFrameNumber KeyTime = BakeData.Value.KeyTimes[Counter];
+				DoubleChannels[0]->AddLinearKey(KeyTime, LocalTranslations[Counter].X);
+				DoubleChannels[1]->AddLinearKey(KeyTime, LocalTranslations[Counter].Y);
+				DoubleChannels[2]->AddLinearKey(KeyTime, LocalTranslations[Counter].Z);
+				DoubleChannels[3]->AddLinearKey(KeyTime, LocalRotations[Counter].X);
+				DoubleChannels[4]->AddLinearKey(KeyTime, LocalRotations[Counter].Y);
+				DoubleChannels[5]->AddLinearKey(KeyTime, LocalRotations[Counter].Z);
+				DoubleChannels[6]->AddLinearKey(KeyTime, LocalScales[Counter].X);
+				DoubleChannels[7]->AddLinearKey(KeyTime, LocalScales[Counter].Y);
+				DoubleChannels[8]->AddLinearKey(KeyTime, LocalScales[Counter].Z);
+			}
+		}
+	}
+	
+	Sequencer->NotifyMovieSceneDataChanged( EMovieSceneDataChangeType::MovieSceneStructureItemsChanged );
+}
+
+void UDaySequenceEditorSubsystem::FixActorReferences()
+{
+	TSharedPtr<ISequencer> Sequencer = GetActiveSequencer();
+	if (Sequencer == nullptr)
+	{
+		return;
+	}
+
+	UWorld* PlaybackContext = Sequencer->GetPlaybackContext()->GetWorld();
+
+	if (!PlaybackContext)
+	{
+		return;
+	}
+
+	UMovieScene* FocusedMovieScene = Sequencer->GetFocusedMovieSceneSequence()->GetMovieScene();
+	if (!FocusedMovieScene)
+	{
+		return;
+	}
+
+	if (FocusedMovieScene->IsReadOnly())
+	{
+		FSequencerUtilities::ShowReadOnlyError();
+		return;
+	}
+
+	FScopedTransaction FixActorReferencesTransaction(LOCTEXT("FixActorReferences", "Fix Actor References"));
+
+	TMap<FString, AActor*> ActorNameToActorMap;
+
+	for (TActorIterator<AActor> ActorItr(PlaybackContext); ActorItr; ++ActorItr)
+	{
+		// Same as with the Object Iterator, access the subclass instance with the * or -> operators.
+		AActor* Actor = *ActorItr;
+		ActorNameToActorMap.Add(Actor->GetActorLabel(), Actor);
+	}
+
+	// Cache the possessables to fix up first since the bindings will change as the fix ups happen.
+	TArray<FMovieScenePossessable> ActorsPossessablesToFix;
+	for (int32 i = 0; i < FocusedMovieScene->GetPossessableCount(); i++)
+	{
+		FMovieScenePossessable& Possessable = FocusedMovieScene->GetPossessable(i);
+		// Possessables with parents are components so ignore them.
+		if (Possessable.GetParent().IsValid() == false)
+		{
+			if (Sequencer->FindBoundObjects(Possessable.GetGuid(), Sequencer->GetFocusedTemplateID()).Num() == 0)
+			{
+				ActorsPossessablesToFix.Add(Possessable);
+			}
+		}
+	}
+
+	// For the possessables to fix, look up the actors by name and reassign them if found.
+	TMap<FGuid, FGuid> OldGuidToNewGuidMap;
+	for (const FMovieScenePossessable& ActorPossessableToFix : ActorsPossessablesToFix)
+	{
+		AActor* ActorPtr = ActorNameToActorMap.FindRef(ActorPossessableToFix.GetName());
+		if (ActorPtr != nullptr)
+		{
+			FGuid OldGuid = ActorPossessableToFix.GetGuid();
+
+			// The actor might have an existing guid while the possessable with the same name might not. 
+			// In that case, make sure we also replace the existing guid with the new guid 
+			FGuid ExistingGuid = Sequencer->FindObjectId(*ActorPtr, Sequencer->GetFocusedTemplateID());
+
+			FGuid NewGuid = FSequencerUtilities::AssignActor(Sequencer.ToSharedRef(), ActorPtr, ActorPossessableToFix.GetGuid());
+
+			OldGuidToNewGuidMap.Add(OldGuid, NewGuid);
+
+			if (ExistingGuid.IsValid())
+			{
+				OldGuidToNewGuidMap.Add(ExistingGuid, NewGuid);
+			}
+		}
+	}
+
+	for (TPair<FGuid, FGuid> GuidPair : OldGuidToNewGuidMap)
+	{
+		FSequencerUtilities::UpdateBindingIDs(Sequencer.ToSharedRef(), GuidPair.Key, GuidPair.Value);
+	}
+}
+
+void UDaySequenceEditorSubsystem::AddActorsToBindingInternal()
+{
+	TSharedPtr<ISequencer> Sequencer = GetActiveSequencer();
+	if (Sequencer == nullptr)
+	{
+		return;
+	}
+
+	TArray<FGuid> ObjectBindings;
+	Sequencer->GetSelectedObjects(ObjectBindings);
+	if (ObjectBindings.Num() == 0)
+	{
+		return;
+	}
+
+	FMovieSceneBindingProxy BindingProxy(ObjectBindings[0], Sequencer->GetFocusedMovieSceneSequence());
+
+	TArray<AActor*> SelectedActors;
+	GEditor->GetSelectedActors()->GetSelectedObjects<AActor>(SelectedActors); 
+	AddActorsToBinding(SelectedActors, BindingProxy);
+}
+
+void UDaySequenceEditorSubsystem::AddActorsToBinding(const TArray<AActor*>& Actors, const FMovieSceneBindingProxy& ObjectBinding)
+{
+	TSharedPtr<ISequencer> Sequencer = GetActiveSequencer();
+	if (Sequencer == nullptr)
+	{
+		return;
+	}
+
+	FSequencerUtilities::AddActorsToBinding(Sequencer.ToSharedRef(), Actors, ObjectBinding);
+}
+
+void UDaySequenceEditorSubsystem::ReplaceBindingWithActorsInternal()
+{
+	TSharedPtr<ISequencer> Sequencer = GetActiveSequencer();
+	if (Sequencer == nullptr)
+	{
+		return;
+	}
+
+	TArray<FGuid> ObjectBindings;
+	Sequencer->GetSelectedObjects(ObjectBindings);
+	if (ObjectBindings.Num() == 0)
+	{
+		return;
+	}
+
+	FMovieSceneBindingProxy BindingProxy(ObjectBindings[0], Sequencer->GetFocusedMovieSceneSequence());
+
+	TArray<AActor*> SelectedActors;
+	GEditor->GetSelectedActors()->GetSelectedObjects<AActor>(SelectedActors); 
+	ReplaceBindingWithActors(SelectedActors, BindingProxy);
+}
+
+void UDaySequenceEditorSubsystem::ReplaceBindingWithActors(const TArray<AActor*>& Actors, const FMovieSceneBindingProxy& ObjectBinding)
+{
+	TSharedPtr<ISequencer> Sequencer = GetActiveSequencer();
+	if (Sequencer == nullptr)
+	{
+		return;
+	}
+
+	FSequencerUtilities::ReplaceBindingWithActors(Sequencer.ToSharedRef(), Actors, ObjectBinding);
+}
+
+void UDaySequenceEditorSubsystem::RemoveActorsFromBindingInternal()
+{
+	TSharedPtr<ISequencer> Sequencer = GetActiveSequencer();
+	if (Sequencer == nullptr)
+	{
+		return;
+	}
+
+	TArray<FGuid> ObjectBindings;
+	Sequencer->GetSelectedObjects(ObjectBindings);
+	if (ObjectBindings.Num() == 0)
+	{
+		return;
+	}
+
+	FMovieSceneBindingProxy BindingProxy(ObjectBindings[0], Sequencer->GetFocusedMovieSceneSequence());
+
+	TArray<AActor*> SelectedActors;
+	GEditor->GetSelectedActors()->GetSelectedObjects<AActor>(SelectedActors); 
+	RemoveActorsFromBinding(SelectedActors, BindingProxy);
+}
+
+void UDaySequenceEditorSubsystem::RemoveActorsFromBinding(const TArray<AActor*>& Actors, const FMovieSceneBindingProxy& ObjectBinding)
+{
+	TSharedPtr<ISequencer> Sequencer = GetActiveSequencer();
+	if (Sequencer == nullptr)
+	{
+		return;
+	}
+
+	FSequencerUtilities::RemoveActorsFromBinding(Sequencer.ToSharedRef(), Actors, ObjectBinding);
+}
+
+void UDaySequenceEditorSubsystem::RemoveAllBindingsInternal()
+{
+	TSharedPtr<ISequencer> Sequencer = GetActiveSequencer();
+	if (Sequencer == nullptr)
+	{
+		return;
+	}
+
+	TArray<FGuid> ObjectBindings;
+	Sequencer->GetSelectedObjects(ObjectBindings);
+	if (ObjectBindings.Num() == 0)
+	{
+		return;
+	}
+
+	FMovieSceneBindingProxy BindingProxy(ObjectBindings[0], Sequencer->GetFocusedMovieSceneSequence());
+
+	RemoveAllBindings(BindingProxy);
+}
+
+void UDaySequenceEditorSubsystem::RemoveAllBindings(const FMovieSceneBindingProxy& ObjectBinding)
+{
+	TSharedPtr<ISequencer> Sequencer = GetActiveSequencer();
+	if (Sequencer == nullptr)
+	{
+		return;
+	}
+
+	UMovieSceneSequence* Sequence = Sequencer->GetFocusedMovieSceneSequence();
+	if (!Sequence)
+	{
+		return;
+	}
+
+	UMovieScene* MovieScene = Sequence->GetMovieScene();
+	if (!MovieScene)
+	{
+		return;
+	}
+
+	FScopedTransaction RemoveAllBindings(LOCTEXT("RemoveAllBindings", "Remove All Bound Objects"));
+
+	Sequence->Modify();
+	MovieScene->Modify();
+
+	// Unbind objects
+	FGuid Guid = ObjectBinding.BindingID;
+	Sequence->UnbindPossessableObjects(Guid);
+
+	Sequencer->RestorePreAnimatedState();
+
+	Sequencer->NotifyMovieSceneDataChanged(EMovieSceneDataChangeType::MovieSceneStructureItemsChanged);
+}
+
+void UDaySequenceEditorSubsystem::RemoveInvalidBindingsInternal()
+{
+	TSharedPtr<ISequencer> Sequencer = GetActiveSequencer();
+	if (Sequencer == nullptr)
+	{
+		return;
+	}
+
+	TArray<FGuid> ObjectBindings;
+	Sequencer->GetSelectedObjects(ObjectBindings);
+	if (ObjectBindings.Num() == 0)
+	{
+		return;
+	}
+
+	FMovieSceneBindingProxy BindingProxy(ObjectBindings[0], Sequencer->GetFocusedMovieSceneSequence());
+
+	RemoveInvalidBindings(BindingProxy);
+}
+
+void UDaySequenceEditorSubsystem::RemoveInvalidBindings(const FMovieSceneBindingProxy& ObjectBinding)
+{
+	TSharedPtr<ISequencer> Sequencer = GetActiveSequencer();
+	if (Sequencer == nullptr)
+	{
+		return;
+	}
+
+	UMovieSceneSequence* Sequence = Sequencer->GetFocusedMovieSceneSequence();
+	if (!Sequence)
+	{
+		return;
+	}
+
+	UMovieScene* MovieScene = Sequence->GetMovieScene();
+	if (!MovieScene)
+	{
+		return;
+	}
+	
+	FScopedTransaction RemoveInvalidBindings(LOCTEXT("RemoveMissing", "Remove Missing Objects"));
+
+	Sequence->Modify();
+	MovieScene->Modify();
+
+	// Unbind objects
+	FGuid Guid = ObjectBinding.BindingID;
+	Sequence->UnbindInvalidObjects(Guid, Sequencer->GetPlaybackContext());
+
+	// Update label
+	UClass* ActorClass = nullptr;
+
+	TArray<AActor*> ValidActors;
+	for (TWeakObjectPtr<> Ptr : Sequencer->FindObjectsInCurrentSequence(Guid))
+	{
+		if (AActor* Actor = Cast<AActor>(Ptr.Get()))
+		{
+			ActorClass = Actor->GetClass();
+			ValidActors.Add(Actor);
+		}
+	}
+
+	FMovieScenePossessable* Possessable = MovieScene->FindPossessable(Guid);
+	if (Possessable && ActorClass != nullptr && ValidActors.Num() != 0)
+	{
+		if (ValidActors.Num() > 1)
+		{
+			FString NewLabel = ActorClass->GetName() + FString::Printf(TEXT(" (%d)"), ValidActors.Num());
+
+			Possessable->SetName(NewLabel);
+		}
+		else
+		{
+			Possessable->SetName(ValidActors[0]->GetActorLabel());
+		}
+	}
+
+	Sequencer->RestorePreAnimatedState();
+
+	Sequencer->NotifyMovieSceneDataChanged(EMovieSceneDataChangeType::MovieSceneStructureItemsChanged);
+}
+
+void UDaySequenceEditorSubsystem::AddAssignActorMenu(FMenuBuilder& MenuBuilder)
+{
+	TSharedPtr<ISequencer> Sequencer = GetActiveSequencer();
+	if (Sequencer == nullptr)
+	{
+		return;
+	}
+
+	TArray<FGuid> ObjectBindings;
+	Sequencer->GetSelectedObjects(ObjectBindings);
+	if (ObjectBindings.Num() == 0)
+	{
+		return;
+	}
+
+	MenuBuilder.AddMenuEntry(FDaySequenceEditorCommands::Get().AddActorsToBinding);
+	MenuBuilder.AddMenuEntry(FDaySequenceEditorCommands::Get().ReplaceBindingWithActors);
+	MenuBuilder.AddMenuEntry(FDaySequenceEditorCommands::Get().RemoveActorsFromBinding);
+	MenuBuilder.AddMenuEntry(FDaySequenceEditorCommands::Get().RemoveAllBindings);
+	MenuBuilder.AddMenuEntry(FDaySequenceEditorCommands::Get().RemoveInvalidBindings);
+
+	FMovieSceneBindingProxy BindingProxy(ObjectBindings[0], Sequencer->GetFocusedMovieSceneSequence());
+
+	TSet<const AActor*> BoundObjects;
+	{
+		for (TWeakObjectPtr<> Ptr : Sequencer->FindObjectsInCurrentSequence(ObjectBindings[0]))
+		{
+			if (const AActor* Actor = Cast<AActor>(Ptr.Get()))
+			{
+				BoundObjects.Add(Actor);
+			}
+		}
+	}
+
+	auto IsActorValidForAssignment = [BoundObjects](const AActor* InActor){
+		return !BoundObjects.Contains(InActor);
+	};
+
+	// Set up a menu entry to assign an actor to the object binding node
+	FSceneOutlinerInitializationOptions InitOptions;
+	{
+		// We hide the header row to keep the UI compact.
+		InitOptions.bShowHeaderRow = false;
+		InitOptions.bShowSearchBox = true;
+		InitOptions.bShowCreateNewFolder = false;
+		InitOptions.bFocusSearchBoxWhenOpened = true;
+		InitOptions.bShowTransient = true;
+		// Only want the actor label column
+		InitOptions.ColumnMap.Add(FSceneOutlinerBuiltInColumnTypes::Label(), FSceneOutlinerColumnInfo(ESceneOutlinerColumnVisibility::Visible, 0));
+		
+		// Only display actors that are not possessed already
+		InitOptions.Filters->AddFilterPredicate<FActorTreeItem>(FActorTreeItem::FFilterPredicate::CreateLambda( IsActorValidForAssignment ) );
+	}
+
+	// actor selector to allow the user to choose an actor
+	FSceneOutlinerModule& SceneOutlinerModule = FModuleManager::LoadModuleChecked<FSceneOutlinerModule>("SceneOutliner");
+	TSharedRef< SWidget > MiniSceneOutliner =
+		SNew( SBox )
+		.MaxDesiredHeight(400.0f)
+		.WidthOverride(300.0f)
+		[
+			SceneOutlinerModule.CreateActorPicker(
+				InitOptions,
+				FOnActorPicked::CreateLambda([=](AActor* Actor){
+					// Create a new binding for this actor
+					FSlateApplication::Get().DismissAllMenus();
+					FSequencerUtilities::AssignActor(Sequencer.ToSharedRef(), Actor, ObjectBindings[0]);
+				})
+			)
+		];
+
+	MenuBuilder.AddMenuSeparator();
+	MenuBuilder.AddWidget(MiniSceneOutliner, FText::GetEmpty(), true);
+}
+
+void UDaySequenceEditorSubsystem::GetRebindComponentNames(TArray<FName>& OutComponentNames)
+{
+	OutComponentNames.Empty();
+
+	TSharedPtr<ISequencer> Sequencer = GetActiveSequencer();
+	if (Sequencer == nullptr)
+	{
+		return;
+	}
+
+	UMovieSceneSequence* Sequence = Sequencer->GetFocusedMovieSceneSequence();
+	if (!Sequence)
+	{
+		return;
+	}
+
+	UMovieScene* MovieScene = Sequence->GetMovieScene();
+	if (!MovieScene)
+	{
+		return;
+	}
+
+	TArray<FGuid> ObjectBindings;
+	Sequencer->GetSelectedObjects(ObjectBindings);
+	if (ObjectBindings.Num() == 0)
+	{
+		return;
+	}
+
+	FGuid ComponentGuid = ObjectBindings[0];
+
+	FMovieScenePossessable* ComponentPossessable = MovieScene->FindPossessable(ComponentGuid);
+
+	FGuid ActorParentGuid = ComponentPossessable ? ComponentPossessable->GetParent() : FGuid();
+
+	TArrayView<TWeakObjectPtr<>> ObjectsInCurrentSequence = Sequencer->FindObjectsInCurrentSequence(ActorParentGuid);
+
+	const AActor* Actor = nullptr;
+	for (TWeakObjectPtr<> Ptr : ObjectsInCurrentSequence)
+	{
+		Actor = Cast<AActor>(Ptr.Get());
+		if (Actor)
+		{
+			break;
+		}
+	}
+
+	if (!Actor)
+	{
+		return;
+	}
+		
+	for (UActorComponent* Component : Actor->GetComponents())
+	{
+		if (Component && ComponentPossessable && Component->GetName() != ComponentPossessable->GetName())
+		{
+			OutComponentNames.Add(Component->GetFName());
+		}
+	}
+	OutComponentNames.Sort(FNameFastLess());
+}
+
+void UDaySequenceEditorSubsystem::RebindComponentMenu(FMenuBuilder& MenuBuilder)
+{
+	TArray<FName> ComponentNames;
+	GetRebindComponentNames(ComponentNames);
+
+	for (const FName& ComponentName : ComponentNames)
+	{
+		FText RebindComponentLabel = FText::FromName(ComponentName);
+		MenuBuilder.AddMenuEntry(
+			RebindComponentLabel, 
+			FText(), 
+			FSlateIcon(), 
+			FUIAction(FExecuteAction::CreateLambda([this, ComponentName]() { RebindComponentInternal(ComponentName); } ) ) );
+	}
+}
+
+void UDaySequenceEditorSubsystem::RebindComponentInternal(const FName& ComponentName)
+{
+	TSharedPtr<ISequencer> Sequencer = GetActiveSequencer();
+	if (Sequencer == nullptr)
+	{
+		return;
+	}
+
+	TArray<FGuid> ObjectBindings;
+	Sequencer->GetSelectedObjects(ObjectBindings);
+	if (ObjectBindings.Num() == 0)
+	{
+		return;
+	}
+
+	TArray<FMovieSceneBindingProxy> BindingProxies;
+	for (const FGuid& ObjectBinding : ObjectBindings)
+	{
+		FMovieSceneBindingProxy BindingProxy(ObjectBinding, Sequencer->GetFocusedMovieSceneSequence());
+		BindingProxies.Add(BindingProxy);
+	}
+
+	RebindComponent(BindingProxies, ComponentName);
+}
+
+void UDaySequenceEditorSubsystem::RebindComponent(const TArray<FMovieSceneBindingProxy>& PossessableBindings, const FName& ComponentName)
+{
+	TSharedPtr<ISequencer> Sequencer = GetActiveSequencer();
+	if (Sequencer == nullptr)
+	{
+		return;
+	}
+
+	UMovieSceneSequence* Sequence = Sequencer->GetFocusedMovieSceneSequence();
+	if (!Sequence)
+	{
+		return;
+	}
+
+	UMovieScene* MovieScene = Sequence->GetMovieScene();
+	if (!MovieScene)
+	{
+		return;
+	}
+
+	FScopedTransaction RebindComponent(LOCTEXT("RebindComponent", "Rebind Component"));
+
+	Sequence->Modify();
+	MovieScene->Modify();
+
+	bool bAnythingChanged = false;
+	for (const FMovieSceneBindingProxy& PossessableBinding : PossessableBindings)
+	{
+		FMovieScenePossessable* ComponentPossessable = MovieScene->FindPossessable(PossessableBinding.BindingID);
+
+		FGuid ActorParentGuid = ComponentPossessable ? ComponentPossessable->GetParent() : FGuid();
+
+		TArrayView<TWeakObjectPtr<>> ObjectsInCurrentSequence = Sequencer->FindObjectsInCurrentSequence(ActorParentGuid);
+
+		for (TWeakObjectPtr<> Ptr : ObjectsInCurrentSequence)
+		{
+			if (const AActor* Actor = Cast<AActor>(Ptr.Get()))
+			{
+				for (UActorComponent* Component : Actor->GetComponents())
+				{
+					if (Component->GetFName() == ComponentName)
+					{
+						// Use GetHandleToObject rather than CreatePossessable to search for an existing
+                        // actor owner of the component binding to ensure that the component is added
+                        // under the existing actor.
+                        FGuid ComponentBinding = Sequencer->GetHandleToObject(Component);
+						
+						if (PossessableBinding.BindingID.IsValid() && ComponentBinding.IsValid())
+						{
+							MovieScene->MoveBindingContents(PossessableBinding.BindingID, ComponentBinding);
+
+							bAnythingChanged = true;
+						}
+					}
+				}
+			}
+		}
+	}
+
+	if (bAnythingChanged)
+	{
+		Sequencer->NotifyMovieSceneDataChanged(EMovieSceneDataChangeType::MovieSceneStructureItemsChanged);
+	}
+}
+
+#undef LOCTEXT_NAMESPACE
+
diff --git a/Engine/Plugins/Experimental/DaySequence/Source/DaySequenceEditor/Public/DaySequenceEditorSpawnRegister.h b/Engine/Plugins/Experimental/DaySequence/Source/DaySequenceEditor/Public/DaySequenceEditorSpawnRegister.h
new file mode 100644
--- /dev/null
+++ b/Engine/Plugins/Experimental/DaySequence/Source/DaySequenceEditor/Public/DaySequenceEditorSpawnRegister.h
@@ -0,0 +1,106 @@
+// Copyright Epic Games, Inc. All Rights Reserved.
+
+#pragma once
+
+#include "Templates/ValueOrError.h"
+#include "MovieSceneSequenceID.h"
+#include "MovieSceneSpawnRegister.h"
+#include "DaySequenceSpawnRegister.h"
+#include "UObject/ObjectKey.h"
+
+struct FNewSpawnable;
+
+class IMovieScenePlayer;
+class ISequencer;
+class UMovieScene;
+class UMovieSceneSequence;
+class FObjectPreSaveContext;
+
+/**
+ * Spawn register used in the editor to add some usability features like maintaining selection states, and projecting spawned state onto spawnable defaults
+ */
+class DAYSEQUENCEEDITOR_API FDaySequenceEditorSpawnRegister
+	: public FDaySequenceSpawnRegister
+{
+public:
+
+	/** Constructor */
+	FDaySequenceEditorSpawnRegister();
+
+	/** Destructor. */
+	~FDaySequenceEditorSpawnRegister() override;
+
+public:
+	void SetSequencer(const TSharedPtr<ISequencer>& Sequencer);
+
+public:
+	// FDaySequenceSpawnRegister interface
+	virtual UObject* SpawnObject(FMovieSceneSpawnable& Spawnable, FMovieSceneSequenceIDRef TemplateID, TSharedRef<const FSharedPlaybackState> SharedPlaybackState) override;
+	virtual void PreDestroyObject(UObject& Object, const FGuid& BindingId, int32 BindingIndex, FMovieSceneSequenceIDRef TemplateID) override;
+#if WITH_EDITOR
+	virtual TValueOrError<FNewSpawnable, FText> CreateNewSpawnableType(UObject& SourceObject, UMovieScene& OwnerMovieScene, UActorFactory* ActorFactory = nullptr) override;
+	virtual void SaveDefaultSpawnableState(FMovieSceneSpawnable& Spawnable, FMovieSceneSequenceIDRef TemplateID, TSharedRef<const FSharedPlaybackState> SharedPlaybackState) override;
+	virtual void SetupDefaultsForSpawnable(UObject* SpawnedObject, const FGuid& Guid, const TOptional<FTransformData>& TransformData, TSharedRef<ISequencer> Sequencer, USequencerSettings* Settings) override;
+	virtual void HandleConvertPossessableToSpawnable(UObject* OldObject, TSharedRef<const FSharedPlaybackState> SharedPlaybackState, TOptional<FTransformData>& OutTransformData) override;
+	virtual bool CanConvertSpawnableToPossessable(FMovieSceneSpawnable& Spawnable) const override;
+#endif
+
+private:
+	/** Called when the editor selection has changed. */
+	void HandleActorSelectionChanged(const TArray<UObject*>& NewSelection, bool bForceRefresh);
+
+#if WITH_EDITOR
+	/** Saves the default state for the specified spawnable, if an instance for it currently exists */
+	void SaveDefaultSpawnableState(const FGuid& Guid, int32 BindingIndex, FMovieSceneSequenceIDRef TemplateID, TSharedRef<const FSharedPlaybackState> SharedPlaybackState) override;
+	void SaveDefaultSpawnableStateImpl(FMovieSceneSpawnable& Spawnable, UMovieSceneSequence* Sequence, UObject* SpawnedObject, TSharedRef<const FSharedPlaybackState> SharedPlaybackState);
+#endif
+
+	/** Called from the editor when a blueprint object replacement has occurred */
+	void OnObjectsReplaced(const TMap<UObject*, UObject*>& OldToNewInstanceMap);
+
+	/** Called whenever an object is modified in the editor */
+	void OnObjectModified(UObject* ModifiedObject);
+
+	/** Called before an object is saved in the editor */
+	void OnPreObjectSaved(UObject* Object, FObjectPreSaveContext SaveContext);
+
+private:
+	struct FTrackedObjectState
+	{
+		FTrackedObjectState(FMovieSceneSequenceIDRef InTemplateID, const FGuid& InObjectBindingID) : TemplateID(InTemplateID), ObjectBindingID(InObjectBindingID), bHasBeenModified(false) {}
+
+		/** The sequence ID that spawned this object */
+		FMovieSceneSequenceID TemplateID;
+
+		/** The object binding ID of the object in the template */
+		FGuid ObjectBindingID;
+
+		/** true if this object has been modified since it was spawned and is different from the current object template */
+		bool bHasBeenModified;
+	};
+
+private:
+	/** Handles for delegates that we've bound to. */
+	FDelegateHandle OnActorSelectionChangedHandle;
+
+	/** Set of spawn register keys for objects that should be selected if they are spawned. */
+	TSet<FMovieSceneSpawnRegisterKey> SelectedSpawnedObjects;
+
+	/** Map from a sequenceID to an array of objects that have been tracked */
+	TMap<FObjectKey, FTrackedObjectState> TrackedObjects;
+
+	/** Set of UMovieSceneSequences that this register has spawned objects for that are modified */
+	TSet<FObjectKey> SequencesWithModifiedObjects;
+
+	/** True if we should clear the above selection cache when the editor selection has been changed. */
+	bool bShouldClearSelectionCache;
+
+	/** Weak pointer to the active sequencer. */
+	TWeakPtr<ISequencer> WeakSequencer;
+
+	/** Handle to a delegate that is bound to FCoreUObjectDelegates::OnObjectModified to harvest changes to spawned objects. */
+	FDelegateHandle OnObjectModifiedHandle;
+
+	/** Handle to a delegate that is bound to FCoreUObjectDelegates::OnObjectPreSave to harvest changes to spawned objects. */
+	FDelegateHandle OnObjectSavedHandle;
+};
diff --git a/Engine/Plugins/Experimental/DaySequence/Source/DaySequenceEditor/Public/DaySequenceEditorSubsystem.h b/Engine/Plugins/Experimental/DaySequence/Source/DaySequenceEditor/Public/DaySequenceEditorSubsystem.h
new file mode 100644
--- /dev/null
+++ b/Engine/Plugins/Experimental/DaySequence/Source/DaySequenceEditor/Public/DaySequenceEditorSubsystem.h
@@ -0,0 +1,188 @@
+// Copyright Epic Games, Inc. All Rights Reserved.
+
+
+#pragma once
+
+#include "EditorSubsystem.h"
+
+#include "MovieSceneTimeUnit.h"
+#include "LevelSequenceEditorSubsystem.h" // For FMovieSceneScriptingParams
+
+#include "DaySequenceEditorSubsystem.generated.h"
+
+DECLARE_LOG_CATEGORY_EXTERN(LogDaySequenceEditor, Log, All);
+
+class FExtender;
+class FMenuBuilder;
+class UMovieSceneCompiledDataManager;
+class UMovieSceneFolder;
+class UMovieSceneSection;
+
+/**
+* UDaySequenceEditorSubsystem
+* Subsystem for day sequence editor related utilities to scripts
+*/
+UCLASS()
+class DAYSEQUENCEEDITOR_API UDaySequenceEditorSubsystem : public UEditorSubsystem
+{
+	GENERATED_BODY()
+
+public:
+
+	virtual void Initialize(FSubsystemCollectionBase& Collection) override;
+	virtual void Deinitialize() override;
+
+	void OnSequencerCreated(TSharedRef<ISequencer> InSequencer);
+
+	/** Add existing actors to Sequencer. Tracks will be automatically added based on default track settings. */
+	UFUNCTION(BlueprintCallable, Category = "Day Sequence Editor")
+	TArray<FMovieSceneBindingProxy> AddActors(const TArray<AActor*>& Actors);
+
+	/** Convert to spawnable. If there are multiple objects assigned to the possessable, multiple spawnables will be created. */
+	UFUNCTION(BlueprintCallable, Category = "Day Sequence Editor")
+	TArray<FMovieSceneBindingProxy> ConvertToSpawnable(const FMovieSceneBindingProxy& ObjectBinding);
+
+	/** Convert to possessable */
+	UFUNCTION(BlueprintCallable, Category = "Day Sequence Editor")
+	FMovieSceneBindingProxy ConvertToPossessable(const FMovieSceneBindingProxy& ObjectBinding);
+
+	/** 
+	 * Copy folders 
+	 * The copied folders will be saved to the clipboard as well as assigned to the ExportedText string. 
+	 * The ExportedTest string can be used in conjunction with PasteFolders if, for example, pasting copy/pasting multiple 
+	 * folders without relying on a single clipboard. 
+	 */
+	UFUNCTION(BlueprintCallable, Category = "Day Sequence Editor")
+	void CopyFolders(const TArray<UMovieSceneFolder*>& Folders, FString& ExportedText);
+
+	/** 
+	 * Paste folders 
+	 * Paste folders from the given TextToImport string (used in conjunction with CopyFolders). 
+	 * If TextToImport is empty, the contents of the clipboard will be used.
+	 */
+	UFUNCTION(BlueprintCallable, Category = "Day Sequence Editor")
+	bool PasteFolders(const FString& TextToImport, FMovieScenePasteFoldersParams PasteFoldersParams, TArray<UMovieSceneFolder*>& OutFolders);
+
+	/**
+	 * Copy sections
+	 * The copied sections will be saved to the clipboard as well as assigned to the ExportedText string.
+	 * The ExportedTest string can be used in conjunction with PasteSections if, for example, pasting copy/pasting multiple
+	 * sections without relying on a single clipboard.
+	 */
+	UFUNCTION(BlueprintCallable, Category = "Day Sequence Editor")
+	void CopySections(const TArray<UMovieSceneSection*>& Sections, FString& ExportedText);
+
+	/**
+	 * Paste sections
+	 * Paste sections from the given TextToImport string (used in conjunction with CopySections).
+	 * If TextToImport is empty, the contents of the clipboard will be used.
+	 */
+	UFUNCTION(BlueprintCallable, Category = "Day Sequence Editor")
+	bool PasteSections(const FString& TextToImport, FMovieScenePasteSectionsParams PasteSectionsParams, TArray<UMovieSceneSection*>& OutSections);
+
+	/**
+	 * Copy tracks
+	 * The copied tracks will be saved to the clipboard as well as assigned to the ExportedText string.
+	 * The ExportedTest string can be used in conjunction with PasteTracks if, for example, pasting copy/pasting multiple
+	 * tracks without relying on a single clipboard.
+	 */
+	UFUNCTION(BlueprintCallable, Category = "Day Sequence Editor")
+	void CopyTracks(const TArray<UMovieSceneTrack*>& Tracks, FString& ExportedText);
+
+	/**
+	 * Paste tracks
+	 * Paste tracks from the given TextToImport string (used in conjunction with CopyTracks).
+	 * If TextToImport is empty, the contents of the clipboard will be used.
+	 */
+	UFUNCTION(BlueprintCallable, Category = "Day Sequence Editor")
+	bool PasteTracks(const FString& TextToImport, FMovieScenePasteTracksParams PasteTracksParams, TArray<UMovieSceneTrack*>& OutTracks);
+
+	/**
+	 * Copy bindings
+	 * The copied bindings will be saved to the clipboard as well as assigned to the ExportedText string.
+	 * The ExportedTest string can be used in conjunction with PasteBindings if, for example, pasting copy/pasting multiple
+	 * bindings without relying on a single clipboard.
+	 */
+	UFUNCTION(BlueprintCallable, Category = "Day Sequence Editor")
+	void CopyBindings(const TArray<FMovieSceneBindingProxy>& Bindings, FString& ExportedText);
+
+	/**
+	 * Paste bindings
+	 * Paste bindings from the given TextToImport string (used in conjunction with CopyBindings).
+	 * If TextToImport is empty, the contents of the clipboard will be used.
+	 */
+	UFUNCTION(BlueprintCallable, Category = "Day Sequence Editor")
+	bool PasteBindings(const FString& TextToImport, FMovieScenePasteBindingsParams PasteBindingsParams, TArray<FMovieSceneBindingProxy>& OutObjectBindings);
+
+	/** Snap sections to timeline using source timecode */
+	UFUNCTION(BlueprintCallable, Category = "Day Sequence Editor")
+	void SnapSectionsToTimelineUsingSourceTimecode(const TArray<UMovieSceneSection*>& Sections);
+
+	/** Sync section using source timecode */
+	UFUNCTION(BlueprintCallable, Category = "Day Sequence Editor")
+	void SyncSectionsUsingSourceTimecode(const TArray<UMovieSceneSection*>& Sections);
+
+	/** Bake transform */
+	UFUNCTION(BlueprintCallable, Category = "Day Sequence Editor")
+	void BakeTransform(const TArray<FMovieSceneBindingProxy>& ObjectBindings, const FFrameTime& BakeInTime, const FFrameTime& BakeOutTime, const FFrameTime& BakeInterval, const FMovieSceneScriptingParams& Params = FMovieSceneScriptingParams());
+
+	/** Attempts to automatically fix up broken actor references in the current scene */
+	UFUNCTION(BlueprintCallable, Category = "Day Sequence Editor")
+	void FixActorReferences();
+
+	/** Assigns the given actors to the binding */
+	UFUNCTION(BlueprintCallable, Category = "Day Sequence Editor")
+	void AddActorsToBinding(const TArray<AActor*>& Actors, const FMovieSceneBindingProxy& ObjectBinding);
+
+	/** Replaces the binding with the given actors */
+	UFUNCTION(BlueprintCallable, Category = "Day Sequence Editor")
+	void ReplaceBindingWithActors(const TArray<AActor*>& Actors, const FMovieSceneBindingProxy& ObjectBinding);
+
+	/** Removes the given actors from the binding */
+	UFUNCTION(BlueprintCallable, Category = "Day Sequence Editor")
+	void RemoveActorsFromBinding(const TArray<AActor*>& Actors, const FMovieSceneBindingProxy& ObjectBinding);
+
+	/** Remove all bound actors from this track */
+	UFUNCTION(BlueprintCallable, Category = "Day Sequence Editor")
+	void RemoveAllBindings(const FMovieSceneBindingProxy& ObjectBinding);
+
+	/** Remove missing objects bound to this track */
+	UFUNCTION(BlueprintCallable, Category = "Day Sequence Editor")
+	void RemoveInvalidBindings(const FMovieSceneBindingProxy& ObjectBinding);
+
+	/** Rebind the component binding to the requested component */
+	UFUNCTION(BlueprintCallable, Category = "Day Sequence Editor")
+	void RebindComponent(const TArray<FMovieSceneBindingProxy>& ComponentBindings, const FName& ComponentName);
+
+private:
+
+	TSharedPtr<ISequencer> GetActiveSequencer();
+	
+	void SnapSectionsToTimelineUsingSourceTimecodeInternal();
+	void SyncSectionsUsingSourceTimecodeInternal();
+	void BakeTransformInternal();
+	void AddActorsToBindingInternal();
+	void ReplaceBindingWithActorsInternal();
+	void RemoveActorsFromBindingInternal();
+	void RemoveAllBindingsInternal();
+	void RemoveInvalidBindingsInternal();
+	void RebindComponentInternal(const FName& ComponentName);
+
+	void AddAssignActorMenu(FMenuBuilder& MenuBuilder);
+
+	void GetRebindComponentNames(TArray<FName>& OutComponentNames);
+	void RebindComponentMenu(FMenuBuilder& MenuBuilder);
+
+	FDelegateHandle OnSequencerCreatedHandle;
+
+	/* List of sequencers that have been created */
+	TArray<TWeakPtr<ISequencer>> Sequencers;
+
+	TSharedPtr<FUICommandList> CommandList;
+
+	TSharedPtr<FExtender> TransformMenuExtender;
+	TSharedPtr<FExtender> FixActorReferencesMenuExtender;
+
+	TSharedPtr<FExtender> AssignActorMenuExtender;
+	TSharedPtr<FExtender> RebindComponentMenuExtender;
+};
