Day Sequence: Remove day sequence editor subsystem

diff --git a/Engine/Plugins/Experimental/DaySequence/Source/DaySequenceEditor/Private/DaySequenceEditorSubsystem.cpp b/Engine/Plugins/Experimental/DaySequence/Source/DaySequenceEditor/Private/DaySequenceEditorSubsystem.cpp
deleted file mode 100644
diff --git a/Engine/Plugins/Experimental/DaySequence/Source/DaySequenceEditor/Private/SDaySequenceConditionTagChip.cpp b/Engine/Plugins/Experimental/DaySequence/Source/DaySequenceEditor/Private/SDaySequenceConditionTagChip.cpp
--- a/Engine/Plugins/Experimental/DaySequence/Source/DaySequenceEditor/Private/SDaySequenceConditionTagChip.cpp
+++ b/Engine/Plugins/Experimental/DaySequence/Source/DaySequenceEditor/Private/SDaySequenceConditionTagChip.cpp
@@ -2,7 +2,6 @@
 
 #include "SDaySequenceConditionTagChip.h"
 
-#include "DaySequenceEditorSubsystem.h"
 #include "Styling/StyleColors.h"
 #include "Widgets/DeclarativeSyntaxSupport.h"
 #include "Widgets/Layout/SBox.h"
@@ -119,7 +118,6 @@ void SDaySequenceConditionTagChip::Construct(const FArguments& InArgs)
 					.ContentPadding(0)
 					.OnClicked_Lambda([WeakSelf]()
 					{
-						UE_LOG(LogDaySequenceEditor, Log, TEXT("Clear button clicked!"))
 						const TSharedPtr<SDaySequenceConditionTagChip> Self = WeakSelf.Pin();
 						if (Self.IsValid() && Self->OnClearPressed.IsBound())
 						{
diff --git a/Engine/Plugins/Experimental/DaySequence/Source/DaySequenceEditor/Public/DaySequenceEditorSubsystem.h b/Engine/Plugins/Experimental/DaySequence/Source/DaySequenceEditor/Public/DaySequenceEditorSubsystem.h
deleted file mode 100644
