SimpleDaySequenceActor: Add optional default SkySphere and material

Also unhides transform category.

diff --git a/Engine/Plugins/Experimental/DaySequence/Source/DaySequence/Private/Actors/SimpleDaySequenceActor.cpp b/Engine/Plugins/Experimental/DaySequence/Source/DaySequence/Private/Actors/SimpleDaySequenceActor.cpp
--- a/Engine/Plugins/Experimental/DaySequence/Source/DaySequence/Private/Actors/SimpleDaySequenceActor.cpp
+++ b/Engine/Plugins/Experimental/DaySequence/Source/DaySequence/Private/Actors/SimpleDaySequenceActor.cpp
@@ -3,12 +3,18 @@
 #include "Actors/SimpleDaySequenceActor.h"
 
 #include "DaySequenceSubsystem.h"
+
 #include "Components/DirectionalLightComponent.h"
 #include "Components/ExponentialHeightFogComponent.h"
 #include "Components/SkyAtmosphereComponent.h"
 #include "Components/SkyLightComponent.h"
 #include "Components/VolumetricCloudComponent.h"
+#include "Components/StaticMeshComponent.h"
+
 #include "Engine/World.h"
+#include "Materials/MaterialInterface.h"
+#include "Engine/StaticMesh.h"
+#include "UObject/ConstructorHelpers.h"
 
 ASimpleDaySequenceActor::ASimpleDaySequenceActor(const FObjectInitializer& Init)
 : Super(Init)
@@ -33,6 +39,14 @@ ASimpleDaySequenceActor::ASimpleDaySequenceActor(const FObjectInitializer& Init)
 
 	VolumetricCloudComponent = CreateOptionalDefaultSubobject<UVolumetricCloudComponent>(TEXT("VolumetricCloud"));
     VolumetricCloudComponent->SetupAttachment(RootComponent);
+
+	static ConstructorHelpers::FObjectFinder<UStaticMesh> SkySphereDefaultMesh(TEXT("/Engine/EngineSky/SM_SkySphere.SM_SkySphere"));
+	static ConstructorHelpers::FObjectFinder<UMaterialInterface> SkySphereDefaultMaterial(TEXT("/Engine/EngineSky/M_SimpleSkyDome.M_SimpleSkyDome"));
+	SkySphereComponent = CreateOptionalDefaultSubobject<UStaticMeshComponent>(TEXT("SkySphere"));
+	SkySphereComponent->SetupAttachment(RootComponent);
+	SkySphereComponent->SetStaticMesh(SkySphereDefaultMesh.Object);
+	SkySphereComponent->SetMaterial(0, SkySphereDefaultMaterial.Object.Get());
+	SkySphereComponent->SetRelativeScale3D(FVector(400.f));
 }
 
 void ASimpleDaySequenceActor::BeginPlay()
diff --git a/Engine/Plugins/Experimental/DaySequence/Source/DaySequence/Public/Actors/SimpleDaySequenceActor.h b/Engine/Plugins/Experimental/DaySequence/Source/DaySequence/Public/Actors/SimpleDaySequenceActor.h
--- a/Engine/Plugins/Experimental/DaySequence/Source/DaySequence/Public/Actors/SimpleDaySequenceActor.h
+++ b/Engine/Plugins/Experimental/DaySequence/Source/DaySequence/Public/Actors/SimpleDaySequenceActor.h
@@ -12,12 +12,13 @@ class USkyLightComponent;
 class UVolumetricCloudComponent;
 class UDirectionalLightComponent;
 class UExponentialHeightFogComponent;
+class UStaticMeshComponent;
 
 /**
  * A self-registering Day Sequence Actor with a simple set of lighting components (some of which are optional).
  * Can be used as-is or extended by derived classes (see AProceduralDaySequenceActor for example).
  */
-UCLASS(Blueprintable, HideCategories=(Tags, Transform, Networking, LevelInstance))
+UCLASS(Blueprintable, HideCategories=(Tags, Networking, LevelInstance))
 class DAYSEQUENCE_API ASimpleDaySequenceActor
 	: public ADaySequenceActor
 {
@@ -53,4 +54,7 @@ protected:
 	
 	UPROPERTY(VisibleDefaultsOnly, BlueprintReadOnly, Category = "Day Sequence", meta = (AllowPrivateAccess = "true"))
 	TObjectPtr<UVolumetricCloudComponent> VolumetricCloudComponent;
+
+	UPROPERTY(VisibleDefaultsOnly, BlueprintReadOnly, Category = "Day Sequence", meta = (AllowPrivateAccess = "true"))
+	TObjectPtr<UStaticMeshComponent> SkySphereComponent;
 };
