DaySequenceModifierComponent: Convert modifier component to use player view position instead of player pawn position.

Modifier blend targets are now expected to be an APlayerController (or a class derived from it). We access the player camera manager to get the view position for the purposes of enabling/disabling and computing blend weight for the modifier.

Instead of relying on overlap events and tracking which pawn the player is currently viewing (which can change in replays, for example), we now poll the player view position every time the sequence player is updated. This solves many issues that were previously only solveable using game/project specific logic.

#changelist validated
#virtualized

diff --git a/Engine/Plugins/Experimental/DaySequence/Source/DaySequence/Private/Actors/DaySequenceModifierVolume.cpp b/Engine/Plugins/Experimental/DaySequence/Source/DaySequence/Private/Actors/DaySequenceModifierVolume.cpp
--- a/Engine/Plugins/Experimental/DaySequence/Source/DaySequence/Private/Actors/DaySequenceModifierVolume.cpp
+++ b/Engine/Plugins/Experimental/DaySequence/Source/DaySequence/Private/Actors/DaySequenceModifierVolume.cpp
@@ -38,16 +38,16 @@ void ADaySequenceModifierVolume::TryEnableModifier() const
 	}
 }
 
-void ADaySequenceModifierVolume::SetBlendTarget(AActor* InBlendTarget)
+void ADaySequenceModifierVolume::SetBlendTarget(APlayerController* InPC)
 {
-	if (!IsValid(InBlendTarget) || InBlendTarget == CurrentBlendTarget)
+	if (!IsValid(InPC) || InPC == CurrentBlendTarget)
 	{
 		return;
 	}
 
-	CurrentBlendTarget = InBlendTarget;
+	CurrentBlendTarget = InPC;
 
-	DaySequenceModifier->EnableDistanceVolumeBlends(InBlendTarget);
+	DaySequenceModifier->EnableDistanceVolumeBlends(InPC);
 	DaySequenceModifier->SetCustomVolumeBlendWeight(1.f);
 
 	TryEnableModifier();
@@ -104,16 +104,6 @@ void ADaySequenceModifierVolume::VolumeSetup()
 
 	// this is kind of a hack, we are calling this for the side effect of forcing the modifier to recache volume shapes after we just invalidated them.
 	DaySequenceModifier->IsBlendTargetInAnyVolume();
-	
-#if WITH_EDITOR
-	const UWorld* World = GetWorld();
-	if (World && World->WorldType != EWorldType::Editor)
-	{
-#endif
-	SetupVolumeCallbacks();
-#if WITH_EDITOR
-	}
-#endif
 }
 
 void ADaySequenceModifierVolume::AddShapeComponentsToModifier()
@@ -121,37 +111,6 @@ void ADaySequenceModifierVolume::AddShapeComponentsToModifier()
 	FComponentReference DefaultBoxReference;
 	DefaultBoxReference.OverrideComponent = DefaultBox;
 	DaySequenceModifier->AddVolumeShapeComponent(DefaultBoxReference);
-
-
-}
-
-void ADaySequenceModifierVolume::SetupVolumeCallbacks()
-{
-	if (!DaySequenceModifier->GetOnVolumeBlendTargetOverlapBegin().IsBoundToObject(this))
-	{
-		DaySequenceModifier->GetOnVolumeBlendTargetOverlapBegin().AddWeakLambda(this, [this](AActor*)
-		{
-			OnOverlapBegin();
-		});
-	}
-	
-	if (!DaySequenceModifier->GetOnVolumeBlendTargetOverlapEnd().IsBoundToObject(this))
-	{
-		DaySequenceModifier->GetOnVolumeBlendTargetOverlapEnd().AddWeakLambda(this, [this](AActor*)
-		{
-			OnOverlapEnd();
-		});
-	}
-}
-
-void ADaySequenceModifierVolume::OnOverlapBegin() const
-{
-	DaySequenceModifier->EnableModifier();
-}
-
-void ADaySequenceModifierVolume::OnOverlapEnd() const
-{
-	DaySequenceModifier->DisableModifier();
 }
 
 void ADaySequenceModifierVolume::PlayerControllerSetup()
@@ -161,7 +120,6 @@ void ADaySequenceModifierVolume::PlayerControllerSetup()
 	if (World && World->WorldType != EWorldType::Editor)
 	{
 #endif
-	SetupBlendTargetCallbacks();
 	CachePlayerController();
 #if WITH_EDITOR
 	}
@@ -192,10 +150,7 @@ void ADaySequenceModifierVolume::CachePlayerController()
 	}
 	else
 	{
-		if (AActor* ViewTarget = CachedPlayerController->GetViewTarget())
-		{
-			SetBlendTarget(ViewTarget);
-		}
+		SetBlendTarget(CachedPlayerController);
 	}
 }
 
@@ -215,30 +170,6 @@ void ADaySequenceModifierVolume::QueuePlayerControllerQuery()
 	}
 }
 
-void ADaySequenceModifierVolume::SetupBlendTargetCallbacks()
-{
-	if (!ViewTargetChangedHandle.IsValid())
-	{
-		ViewTargetChangedHandle = FGameDelegates::Get().GetViewTargetChangedDelegate().AddWeakLambda(this, [this](APlayerController* PC, AActor* OldTarget, AActor* NewTarget)
-		{
-			if (!IsValid(CachedPlayerController))
-			{
-				QueuePlayerControllerQuery();
-				return;
-			}
-
-			if (PC != CachedPlayerController)
-			{
-				return;
-			}
-
-			// there is a lot of extra logic we can do here to stop popping (for example, try cast NewTarget to APlayerController which happens when view target unset).
-			// for now just set
-			SetBlendTarget(NewTarget);
-		});
-	}
-}
-
 void ADaySequenceModifierVolume::DaySequenceActorSetup()
 {
 	SetupDaySequenceSubsystemCallbacks();
diff --git a/Engine/Plugins/Experimental/DaySequence/Source/DaySequence/Private/DaySequenceActor.cpp b/Engine/Plugins/Experimental/DaySequence/Source/DaySequence/Private/DaySequenceActor.cpp
--- a/Engine/Plugins/Experimental/DaySequence/Source/DaySequence/Private/DaySequenceActor.cpp
+++ b/Engine/Plugins/Experimental/DaySequence/Source/DaySequence/Private/DaySequenceActor.cpp
@@ -31,6 +31,8 @@
 #include "Engine/Canvas.h"	// UCanvas
 #endif
 
+#include "Engine/World.h"
+
 #include <limits>
 
 #if WITH_EDITOR
@@ -840,6 +842,8 @@ void ADaySequenceActor::OnSequencePlayerUpdate(const UMovieSceneSequencePlayer&
 	const float CurrentHours = FrameTimeToDayHours(CurrentTime);
 	const float PreviousHours = FrameTimeToDayHours(PreviousTime);
 	SequencePlayerUpdated(CurrentHours, PreviousHours);
+	
+	OnSequencePlayerUpdated.Broadcast();
 }
 
 void ADaySequenceActor::SequencePlayerUpdated(float CurrentTime, float PreviousTime)
diff --git a/Engine/Plugins/Experimental/DaySequence/Source/DaySequence/Private/DaySequenceModifierComponent.cpp b/Engine/Plugins/Experimental/DaySequence/Source/DaySequence/Private/DaySequenceModifierComponent.cpp
--- a/Engine/Plugins/Experimental/DaySequence/Source/DaySequence/Private/DaySequenceModifierComponent.cpp
+++ b/Engine/Plugins/Experimental/DaySequence/Source/DaySequence/Private/DaySequenceModifierComponent.cpp
@@ -269,15 +269,15 @@ UDaySequenceModifierComponent::UDaySequenceModifierComponent(const FObjectInitia
 		(*DebugData).FindOrAdd("Component Enabled") = bIsComponentEnabled ? "True" : "False";
 		(*DebugData).FindOrAdd("Modifier Enabled") = bIsEnabled ? "True" : "False";
 		(*DebugData).FindOrAdd("Blend Weight") = FString::Printf(TEXT("%.5f"), GetCurrentBlendWeight());
-		(*DebugData).FindOrAdd("Blend Target" ) = ExternalVolumeBlendTarget.IsValid() ? ExternalVolumeBlendTarget->GetName() : "None";
+
+		const APlayerController* BlendTarget = ExternalVolumeBlendTarget.Get();
+		(*DebugData).FindOrAdd("Blend Target" ) = BlendTarget ? BlendTarget->GetName() : "None";
 
 		return DebugData;
 	});
 #endif
 	
-	PrimaryComponentTick.bCanEverTick = true;
-	PrimaryComponentTick.bStartWithTickEnabled = false;
-	PrimaryComponentTick.TickGroup = TG_PostPhysics;
+	PrimaryComponentTick.bCanEverTick = false;
 	
 	EasingFunction = CreateDefaultSubobject<UDaySequenceModifierEasingFunction>("EasingFunction", true);
 }
@@ -404,40 +404,35 @@ void UDaySequenceModifierComponent::OnUnregister()
 	RemoveSubSequenceTrack();
 }
 
-void UDaySequenceModifierComponent::BeginPlay()
+void UDaySequenceModifierComponent::SequencePlayerUpdated()
 {
-	Super::BeginPlay();
+	CSV_SCOPED_TIMING_STAT_EXCLUSIVE(DaySequenceModifier_SequencePlayerUpdated);
 	
-	BindOverlapEvents();
-}
-
-void UDaySequenceModifierComponent::TickComponent(float DeltaTime, ELevelTick TickType, FActorComponentTickFunction* ThisTickFunction)
-{
-	Super::TickComponent(DeltaTime, TickType, ThisTickFunction);
+	// Force expensive update
+	const float DistanceBlendFactor = UpdateBlendWeight();
 
-	// If we don't need to tick, don't tick
-	AActor* BlendTarget = ExternalVolumeBlendTarget.Get();
-	if (!BlendTarget || !TargetActor || TargetActor->IsPlaying())
+	if (bIsComponentEnabled && bUseVolume)
 	{
-		PrimaryComponentTick.SetTickFunctionEnable(false);
-		return;
+		if (DistanceBlendFactor > UE_SMALL_NUMBER)
+		{
+			EnableModifier();
+		}
+		else
+		{
+			DisableModifier();
+		}
 	}
+}
 
-	// Update TargetActor if our blend factor has changed
-	const float DistanceBlendFactor = GetDistanceBlendFactor(TargetActor->GetActorLocation());
-	const float BlendFactor = FMath::Min(DistanceBlendFactor, CustomVolumeBlendWeight);
-	if (DistanceBlendFactor > UE_SMALL_NUMBER && CachedBlendFactor != BlendFactor)
-	{
-		TargetActor->SetTimeOfDay(TargetActor->GetTimeOfDay());
-	}
-	CachedBlendFactor = BlendFactor;
+void UDaySequenceModifierComponent::BeginPlay()
+{
+	Super::BeginPlay();
 }
 
 void UDaySequenceModifierComponent::EndPlay(EEndPlayReason::Type Reason)
 {
 	Super::EndPlay(Reason);
 	
-	UnbindOverlapEvents();
 	RemoveSubSequenceTrack();
 }
 
@@ -484,6 +479,7 @@ void UDaySequenceModifierComponent::BindToDaySequenceActor(ADaySequenceActor* Da
 	if (ensureMsgf(DaySequenceActor, TEXT("BindToDaySequenceActor called with a null Day Sequence Actor.")))
 	{
 		DaySequenceActor->GetOnPostInitializeDaySequences().AddUObject(this, &UDaySequenceModifierComponent::ReinitializeSubSequence);
+		DaySequenceActor->GetOnSequencePlayerUpdated().AddUObject(this, &UDaySequenceModifierComponent::SequencePlayerUpdated);
 #if ENABLE_DRAW_DEBUG
 		if (!DaySequenceActor->IsDebugCategoryRegistered(ShowDebug_ModifierCategory))
 		{
@@ -504,6 +500,7 @@ void UDaySequenceModifierComponent::UnbindFromDaySequenceActor()
 	if (TargetActor)
 	{
 		TargetActor->GetOnPostInitializeDaySequences().RemoveAll(this);
+		TargetActor->GetOnSequencePlayerUpdated().RemoveAll(this);
 #if ENABLE_DRAW_DEBUG
 		TargetActor->GetOnDebugLevelChanged().RemoveAll(this);
 		TargetActor->UnregisterDebugEntry(DebugEntry, ShowDebug_ModifierCategory);
@@ -561,7 +558,7 @@ bool UDaySequenceModifierComponent::CanBeEnabled() const
 	if (bUseVolume)
 	{
 		ENetMode NetMode = Actor->GetNetMode();
-		return NetMode == NM_Standalone || NetMode == NM_Client;
+		return NetMode != NM_DedicatedServer;
 	}
 
 	return true;
@@ -621,13 +618,6 @@ void UDaySequenceModifierComponent::EnableModifier()
 	if (TargetActor && !TargetActor->IsPlaying())
 	{
 		TargetActor->SetTimeOfDay(TargetActor->GetTimeOfDay());
-
-		// If the actor is not playing and we are using a blend target, enable ticking on this component
-		// since it is needed to update the target actor if the blend amount changes
-		if (ExternalVolumeBlendTarget.IsValid())
-		{
-			PrimaryComponentTick.SetTickFunctionEnable(true);
-		}
 	}
 
 	OnPostEnableModifier.Broadcast();
@@ -648,7 +638,6 @@ void UDaySequenceModifierComponent::DisableModifier()
 	}
 	
 	bIsEnabled = false;
-	PrimaryComponentTick.SetTickFunctionEnable(false);
 
 	if (TargetActor && !TargetActor->HasAnyFlags(RF_BeginDestroyed))
 	{
@@ -1344,22 +1333,46 @@ void UDaySequenceModifierComponent::SetUserDaySequence(UDaySequence* InDaySequen
 	ReinitializeSubSequence(nullptr);
 }
 
-float UDaySequenceModifierComponent::GetDistanceBlendFactor(FVector Position) const
+bool UDaySequenceModifierComponent::GetBlendPosition(FVector& InPosition) const
 {
-	auto GetDistanceBlendFactorForShape = [this, Position](const UShapeComponent* Shape)
+	CSV_SCOPED_TIMING_STAT_EXCLUSIVE(DaySequenceModifier_GetBlendPosition);
+	
+#if WITH_EDITOR
+	if (!GetWorld()->IsGameWorld())
 	{
-		const float Distance = UE::DaySequence::ComputeSignedDistance(Shape, Position);
-		return Distance < 0.f ? FMath::Clamp(-Distance / BlendAmount, 0.f, 1.f) : 0.f;
-	};
+		InPosition = UE::DaySequence::GVolumePreviewLocation;
+		return true;
+	}
+	else
+#endif
+	if (const APlayerController* BlendTarget = ExternalVolumeBlendTarget.Get())
+	{
+		CSV_SCOPED_TIMING_STAT_EXCLUSIVE(DaySequenceModifier_GetPlayerViewPointScope);
+		InPosition = BlendTarget->PlayerCameraManager->GetCameraLocation();
+		return true;
+	}
 
-	float DistanceBlendFactor = 0.f;
+	return false;
+}
+
+float UDaySequenceModifierComponent::GetDistanceBlendFactorForShape(const UShapeComponent* Shape, const FVector& Position) const
+{
+	const float Distance = UE::DaySequence::ComputeSignedDistance(Shape, Position);
+	return Distance < 0.f ? FMath::Clamp(-Distance / BlendAmount, 0.f, 1.f) : 0.f;
+}
+
+float UDaySequenceModifierComponent::GetDistanceBlendFactor(const FVector& Position) const
+{
+	CSV_SCOPED_TIMING_STAT_EXCLUSIVE(DaySequenceModifier_GetDistanceBlendFactor);
+	
+	CachedDistanceBlendFactor = 0.f;
 
 	for (const UShapeComponent* Shape : GetVolumeShapeComponents())
 	{
-		DistanceBlendFactor = FMath::Max(DistanceBlendFactor, GetDistanceBlendFactorForShape(Shape));
+		CachedDistanceBlendFactor = FMath::Max(CachedDistanceBlendFactor, GetDistanceBlendFactorForShape(Shape, Position));
 	}
 
-	return DistanceBlendFactor;
+	return CachedDistanceBlendFactor;
 }
 
 TArray<UShapeComponent*> UDaySequenceModifierComponent::GetVolumeShapeComponents() const
@@ -1415,25 +1428,14 @@ TArray<UShapeComponent*> UDaySequenceModifierComponent::GetVolumeShapeComponents
 
 float UDaySequenceModifierComponent::GetCurrentBlendWeight() const
 {
-	FVector Position;
-
-#if WITH_EDITOR
-	if (!GetWorld()->IsGameWorld())
-	{
-		Position = UE::DaySequence::GVolumePreviewLocation;
-	}
-	else
-#endif
-	if (const AActor* BlendTarget = ExternalVolumeBlendTarget.Get())
-	{
-		Position = BlendTarget->GetActorLocation();
-	}
-	else
-	{
-		return FMath::Min(1.f, CustomVolumeBlendWeight);
-	}
+	return FMath::Min(CachedDistanceBlendFactor, CustomVolumeBlendWeight);
+}
 
-	return FMath::Min(GetDistanceBlendFactor(Position), CustomVolumeBlendWeight);
+float UDaySequenceModifierComponent::UpdateBlendWeight() const
+{
+	FVector Position;
+	const float DistanceBlendFactor = GetBlendPosition(Position) ? GetDistanceBlendFactor(Position) : 1.f;
+	return FMath::Min(DistanceBlendFactor, CustomVolumeBlendWeight);
 }
 
 void UDaySequenceModifierComponent::SetVolumeCollisionEnabled(const ECollisionEnabled::Type InCollisionType) const
@@ -1461,8 +1463,8 @@ void UDaySequenceModifierComponent::InvalidateMuteStates() const
 	OnInvalidateMuteStates.Broadcast();
 }
 
-void UDaySequenceModifierComponent::EnableDistanceVolumeBlends(AActor* InActor)
-{
+void UDaySequenceModifierComponent::EnableDistanceVolumeBlends(APlayerController* InActor)
+{	
 	ExternalVolumeBlendTarget = InActor;
 }
 
@@ -1504,11 +1506,14 @@ bool UDaySequenceModifierComponent::IsBlendTargetInAnyVolume()
 {
 	OccupiedVolumes = 0;
 
-	for (const UShapeComponent* Shape : GetVolumeShapeComponents())
+	if (FVector Position; GetBlendPosition(Position))
 	{
-		if (ExternalVolumeBlendTarget.IsValid() && Shape->IsOverlappingActor(ExternalVolumeBlendTarget.Get()))
+		for (const UShapeComponent* Shape : GetVolumeShapeComponents())
 		{
-			OccupiedVolumes++;
+			if (GetDistanceBlendFactorForShape(Shape, Position) > 0.f)
+			{
+				OccupiedVolumes++;
+			}
 		}
 	}
 
@@ -1533,57 +1538,6 @@ void UDaySequenceModifierComponent::UpdateCachedExternalShapes() const
 	}
 
 	bCachedExternalShapesInvalid = false;
-	
-	UnbindOverlapEvents();
-	BindOverlapEvents();
-}
-
-void UDaySequenceModifierComponent::BindOverlapEvents() const
-{
-	for (UShapeComponent* Shape : GetVolumeShapeComponents())
-	{
-		// AddUniqueDynamic prevents accidental duplicate binding
-        Shape->OnComponentBeginOverlap.AddUniqueDynamic(this, &UDaySequenceModifierComponent::OnVolumeOverlapBegin);
-        Shape->OnComponentEndOverlap.AddUniqueDynamic(this, &UDaySequenceModifierComponent::OnVolumeOverlapEnd);
-	}
 }
 
-void UDaySequenceModifierComponent::UnbindOverlapEvents() const
-{
-	for (UShapeComponent* Shape : GetVolumeShapeComponents())
-	{
-		Shape->OnComponentBeginOverlap.RemoveDynamic(this, &UDaySequenceModifierComponent::OnVolumeOverlapBegin);
-		Shape->OnComponentEndOverlap.RemoveDynamic(this, &UDaySequenceModifierComponent::OnVolumeOverlapEnd);
-	}
-}
-
-void UDaySequenceModifierComponent::OnVolumeOverlapBegin(UPrimitiveComponent* OverlappedComponent, AActor* OtherActor, UPrimitiveComponent* OtherComp, int32 OtherBodyIndex, bool bFromSweep, const FHitResult& SweepResult)
-{
-	if (ExternalVolumeBlendTarget.IsValid() && OtherActor == ExternalVolumeBlendTarget)
-	{
-		OccupiedVolumes++;
-
-		if (OccupiedVolumes == 1 && GetOwner() && !GetOwner()->HasAuthority())
-		{
-			OnVolumeBlendTargetOverlapBegin.Broadcast(ExternalVolumeBlendTarget.Get());
-			OnVolumeBlendTargetOverlapBeginNative.Broadcast(ExternalVolumeBlendTarget.Get());
-		}
-	}
-}
-
-void UDaySequenceModifierComponent::OnVolumeOverlapEnd(UPrimitiveComponent* OverlappedComponent, AActor* OtherActor, UPrimitiveComponent* OtherComp, int32 OtherBodyIndex)
-{
-	if (ExternalVolumeBlendTarget.IsValid() && OtherActor == ExternalVolumeBlendTarget)
-	{
-		OccupiedVolumes--;
-
-		if (OccupiedVolumes == 0 && GetOwner() && !GetOwner()->HasAuthority())
-		{
-			OnVolumeBlendTargetOverlapEnd.Broadcast(ExternalVolumeBlendTarget.Get());
-			OnVolumeBlendTargetOverlapEndNative.Broadcast(ExternalVolumeBlendTarget.Get());
-		}
-	}
-}
-
-
 #undef LOCTEXT_NAMESPACE
diff --git a/Engine/Plugins/Experimental/DaySequence/Source/DaySequence/Public/Actors/DaySequenceModifierVolume.h b/Engine/Plugins/Experimental/DaySequence/Source/DaySequence/Public/Actors/DaySequenceModifierVolume.h
--- a/Engine/Plugins/Experimental/DaySequence/Source/DaySequence/Public/Actors/DaySequenceModifierVolume.h
+++ b/Engine/Plugins/Experimental/DaySequence/Source/DaySequence/Public/Actors/DaySequenceModifierVolume.h
@@ -52,11 +52,6 @@ protected:
 	 */
 	virtual void AddShapeComponentsToModifier();
 
-	/** Binds the day sequence modifier overlap events to OnOverlapBegin and OnOverlapEnd. */
-	void SetupVolumeCallbacks();
-	void OnOverlapBegin() const;
-	void OnOverlapEnd() const;
-
 	
 	
 	/** Player Controller Setup Functions */
@@ -76,9 +71,6 @@ protected:
 
 	/** Set a timer to call CachePlayerController one tick from now. */
 	void QueuePlayerControllerQuery();
-
-	/** Register a callback to handle changes to our player controller's view target. */
-	void SetupBlendTargetCallbacks();
 	
 
 	
@@ -103,7 +95,7 @@ protected:
 	/** Enables this modifier if the current blend target is in the volume. */
 	void TryEnableModifier() const;
 	
-	void SetBlendTarget(AActor* InPawn);
+	void SetBlendTarget(APlayerController* InPC);
 	
 protected:
 	UPROPERTY(VisibleAnywhere, Category = "Day Sequence", BlueprintReadOnly, meta = (AllowPrivateAccess = "true"))
diff --git a/Engine/Plugins/Experimental/DaySequence/Source/DaySequence/Public/DaySequenceActor.h b/Engine/Plugins/Experimental/DaySequence/Source/DaySequence/Public/DaySequenceActor.h
--- a/Engine/Plugins/Experimental/DaySequence/Source/DaySequence/Public/DaySequenceActor.h
+++ b/Engine/Plugins/Experimental/DaySequence/Source/DaySequence/Public/DaySequenceActor.h
@@ -281,6 +281,9 @@ public:
 	DECLARE_EVENT_OneParam(ADaySequenceActor, FOnPostInitializeDaySequences, FSubSectionPreserveMap*)
 	FOnPostInitializeDaySequences& GetOnPostInitializeDaySequences() { return OnPostInitializeDaySequences; }
 
+	DECLARE_EVENT(ADaySequenceActor, FOnSequencePlayerUpdated)
+	FOnSequencePlayerUpdated& GetOnSequencePlayerUpdated() { return OnSequencePlayerUpdated; }
+	
 	void InvalidateMuteStates() const;
 	
 #if ENABLE_DRAW_DEBUG
@@ -515,6 +518,7 @@ protected:
 	FOnRootSequenceChanged OnPreRootSequenceChanged;
 	FOnRootSequenceChanged OnPostRootSequenceChanged;
 	FOnPostInitializeDaySequences OnPostInitializeDaySequences;
+	FOnSequencePlayerUpdated OnSequencePlayerUpdated;
 	UE::DaySequence::FOnInvalidateMuteStates OnInvalidateMuteStates;
 	
 #if ENABLE_DRAW_DEBUG
diff --git a/Engine/Plugins/Experimental/DaySequence/Source/DaySequence/Public/DaySequenceModifierComponent.h b/Engine/Plugins/Experimental/DaySequence/Source/DaySequence/Public/DaySequenceModifierComponent.h
--- a/Engine/Plugins/Experimental/DaySequence/Source/DaySequence/Public/DaySequenceModifierComponent.h
+++ b/Engine/Plugins/Experimental/DaySequence/Source/DaySequence/Public/DaySequenceModifierComponent.h
@@ -104,8 +104,6 @@ private:
 
 DECLARE_DYNAMIC_MULTICAST_DELEGATE(FOnPostReinitializeSubSequences);
 DECLARE_DYNAMIC_MULTICAST_DELEGATE(FOnPostEnableModifier);
-DECLARE_DYNAMIC_MULTICAST_DELEGATE_OneParam(FOnVolumeBlendTargetOverlap, AActor*, OverlappingActor);
-DECLARE_MULTICAST_DELEGATE_OneParam(FOnVolumeBlendTargetOverlapNative, AActor*);
 
 UCLASS(MinimalAPI, BlueprintType, Blueprintable, config=Game, HideCategories=(Physics,Navigation,Collision,HLOD,Rendering,Cooking,Mobile,RayTracing,AssetUserData), meta=(BlueprintSpawnableComponent))
 class UDaySequenceModifierComponent
@@ -283,7 +281,7 @@ public:
 	 * FinalWeight = Min(DistanceVolumeBlendWeight, CustomVolumeBlendWeight).
 	 */
 	UFUNCTION(BlueprintCallable, Category="Day Sequence")
-	DAYSEQUENCE_API void EnableDistanceVolumeBlends(AActor* InActor);
+	DAYSEQUENCE_API void EnableDistanceVolumeBlends(APlayerController* InActor);
 
 	/**
 	 * Determines the number of volumes the current blend target is
@@ -304,6 +302,8 @@ public:
 	 */
 	UFUNCTION(BlueprintCallable, Category="Day Sequence")
 	float GetCurrentBlendWeight() const;
+	
+	float UpdateBlendWeight() const;
 
 	/**
 	 * Applies the given collision type to all valid volume shape components.
@@ -315,9 +315,6 @@ public:
 	void AddVolumeShapeComponent(const FComponentReference& InShapeReference);
 
 	void InvalidateMuteStates() const;
-	
-	FOnVolumeBlendTargetOverlapNative& GetOnVolumeBlendTargetOverlapBegin() { return OnVolumeBlendTargetOverlapBeginNative; }
-	FOnVolumeBlendTargetOverlapNative& GetOnVolumeBlendTargetOverlapEnd() { return OnVolumeBlendTargetOverlapEndNative; }
 
 #if ENABLE_DRAW_DEBUG
 	bool ShouldShowDebugInfo() const;
@@ -327,12 +324,13 @@ public:
 
 	/*~ Begin UActorComponent interface */
 	void BeginPlay() override;
-	void TickComponent(float DeltaTime, ELevelTick TickType, FActorComponentTickFunction* ThisTickFunction) override;
 	void EndPlay(EEndPlayReason::Type Reason) override;
 	void OnRegister() override;
 	void OnUnregister() override;
 	/*~ End UActorComponent interface */
 
+	void SequencePlayerUpdated();
+	
 #if WITH_EDITOR
 	/*~ Begin FTickableGameObject interface */
 	void UpdateEditorPreview(float DeltaTime) override;
@@ -372,7 +370,7 @@ protected:
 
 	/** The actor to use for distance-based volume blend calculations */
 	UPROPERTY(Transient)
-	TWeakObjectPtr<AActor> ExternalVolumeBlendTarget;
+	TWeakObjectPtr<APlayerController> ExternalVolumeBlendTarget;
 
 	/** An optional user-provided Day Sequence - used instead of our procedurally generated one if set */
 	UPROPERTY(BlueprintReadWrite, EditAnywhere, Category="Day Sequence", meta=(EditCondition="!bUseCollection", EditConditionHides, DisplayAfter="bUseCollection"))
@@ -419,22 +417,6 @@ protected:
 	/** Blueprint exposed delegate invoked after the modifier component is enabled. */
 	UPROPERTY(BlueprintAssignable, Transient, meta=(AllowPrivateAccess="true"))
 	FOnPostEnableModifier OnPostEnableModifier;
-
-	/** 
-	 * Blueprint exposed delegate invoked after the current blend target enters a volume and was not already in a volume.
-	 * Only broadcasts when owning actor does not have network authority.
-	 */
-	UPROPERTY(BlueprintAssignable, Transient, meta=(AllowPrivateAccess="true"))
-	FOnVolumeBlendTargetOverlap OnVolumeBlendTargetOverlapBegin;
-	FOnVolumeBlendTargetOverlapNative OnVolumeBlendTargetOverlapBeginNative;
-	
-	/**
-	 * Blueprint exposed delegate invoked after the current blend target has exited all volumes.
-	 * Only broadcasts when owning actor does not have network authority.
-	 */
-	UPROPERTY(BlueprintAssignable, Transient, meta=(AllowPrivateAccess="true"))
-	FOnVolumeBlendTargetOverlap OnVolumeBlendTargetOverlapEnd;
-	FOnVolumeBlendTargetOverlapNative OnVolumeBlendTargetOverlapEndNative;
 	
 	/** When enabled, these overrides will always override all settings regardless of their bias */
 	UPROPERTY(BlueprintReadWrite, EditAnywhere, Category="Day Sequence", meta=(InlineEditConditionToggle))
@@ -464,8 +446,16 @@ protected:
 	uint8 bUseCollection : 1;
 
 private:
-	float GetDistanceBlendFactor(FVector Position) const;
+	/**
+	 * Get the blend position (handles preview and game world).
+	 * @return Returns true if we have a valid blend position and InPosition was set, false if InPosition was not set.
+	 */
+	bool GetBlendPosition(FVector& InPosition) const;
+	float GetDistanceBlendFactorForShape(const UShapeComponent* Shape, const FVector& Position) const;
 
+	mutable float CachedDistanceBlendFactor;
+	float GetCachedDistanceBlendFactor() const { return CachedDistanceBlendFactor; }
+	float GetDistanceBlendFactor(const FVector& Position) const;
 	/**
 	 * Creates and adds or marks for preserve all subsections that this modifier is responsible for.
 	 * Optionally provided a map of all sections that exist in the root sequence to a bool flag used to mark that section as still relevant.
@@ -478,16 +468,8 @@ private:
 	void UpdateCachedExternalShapes() const;
 	mutable TArray<TWeakObjectPtr<UShapeComponent>> CachedExternalShapes;
 	mutable bool bCachedExternalShapesInvalid = true;
-
-	void BindOverlapEvents() const;
-	void UnbindOverlapEvents() const;
-	int32 OccupiedVolumes = 0;
 	
-	UFUNCTION()
-	void OnVolumeOverlapBegin(UPrimitiveComponent* OverlappedComponent, AActor* OtherActor, UPrimitiveComponent* OtherComp, int32 OtherBodyIndex, bool bFromSweep, const FHitResult& SweepResult);
-
-	UFUNCTION()
-	void OnVolumeOverlapEnd(UPrimitiveComponent* OverlappedComponent, AActor* OtherActor, UPrimitiveComponent* OtherComp, int32 OtherBodyIndex);
+	int32 OccupiedVolumes = 0;
 	
 	/*~ Transient state for active gameplay */
 	TWeakObjectPtr<UMovieSceneSubSection> WeakSubSection;
