SunPositionDaySequenceActor: Change default value of KeyCount to 24 so that it actually works when added to a level

diff --git a/Engine/Plugins/Experimental/DaySequence/Source/DaySequence/Public/Actors/SunPositionDaySequenceActor.h b/Engine/Plugins/Experimental/DaySequence/Source/DaySequence/Public/Actors/SunPositionDaySequenceActor.h
--- a/Engine/Plugins/Experimental/DaySequence/Source/DaySequence/Public/Actors/SunPositionDaySequenceActor.h
+++ b/Engine/Plugins/Experimental/DaySequence/Source/DaySequence/Public/Actors/SunPositionDaySequenceActor.h
@@ -32,7 +32,7 @@ private:
 	
 	/** User settings */
 	UPROPERTY(EditAnywhere, Category = "Procedural Parameters")
-	unsigned KeyCount;
+	unsigned KeyCount = 24;
 	
 	UPROPERTY(EditAnywhere, Category = "Procedural Parameters")
 	FDateTime Time;
