DaySequence: Some minor bug fixes.

- Handle reinstancing of condition tags and configuration of condition tags on load.
- Fix conditions not working on load because we don't support condition instantiation on load. We rebuild entirely after load on tick.
- Fix TimePerCycle changes (we must recreate RootSequence on change)
- Fix rare crash that can occur when ending PIE

diff --git a/Engine/Plugins/Experimental/DaySequence/Source/DaySequence/Private/DaySequenceActor.cpp b/Engine/Plugins/Experimental/DaySequence/Source/DaySequence/Private/DaySequenceActor.cpp
--- a/Engine/Plugins/Experimental/DaySequence/Source/DaySequence/Private/DaySequenceActor.cpp
+++ b/Engine/Plugins/Experimental/DaySequence/Source/DaySequence/Private/DaySequenceActor.cpp
@@ -167,6 +167,10 @@ ADaySequenceActor::ADaySequenceActor(const FObjectInitializer& Init)
 		AHUD::OnShowDebugInfo.AddUObject(this, &ADaySequenceActor::OnShowDebugInfo);
 	}
 #endif
+
+#if WITH_EDITOR
+	FCoreUObjectDelegates::OnObjectsReinstanced.AddUObject(this, &ADaySequenceActor::HandleConditionReinstanced);
+#endif
 }
 
 void ADaySequenceActor::PostInitializeComponents()
@@ -285,6 +289,10 @@ void ADaySequenceActor::PostLoad()
 	// hierarchy for editing binding overrides. This is only necessary for editor, since the
 	// root sequence will be initialized in PostInitializeComponents()/InitializePlayer() for runtime.
 	InitializeRootSequence();
+
+	SubSections.Empty();
+
+	bUpdateRootSequenceOnTick = true;
 #endif
 
 #if WITH_EDITORONLY_DATA
@@ -434,6 +442,9 @@ void ADaySequenceActor::PostEditChangeProperty(FPropertyChangedEvent& PropertyCh
 	{
 		SetTimePerCycle(GetTimePerCycle());
 
+		RootSequence = nullptr;
+		SubSections.Empty();
+		
 		// Update our root sequence with the new asset.
 		UpdateRootSequence();
 	}
@@ -1637,3 +1648,25 @@ void ADaySequenceActor::BindToConditionCallbacks(UObject* LifetimeObject, const
 		}
 	}
 }
+
+#if WITH_EDITOR
+void ADaySequenceActor::HandleConditionReinstanced(const FCoreUObjectDelegates::FReplacementObjectMap& OldToNewInstanceMap)
+{
+	for (const TPair<UObject*, UObject*>& Pair : OldToNewInstanceMap)
+	{
+		if (Pair.Key->IsTemplate())
+		{
+			continue;
+		}
+
+		// Casting too much here? Second one might be unnecessary
+		if (UDaySequenceConditionTag* OldTag = Cast<UDaySequenceConditionTag>(Pair.Key))
+		{
+			if (UDaySequenceConditionTag* NewTag = Cast<UDaySequenceConditionTag>(Pair.Value))
+			{
+				NewTag->GetOnConditionValueChanged() = OldTag->GetOnConditionValueChanged();
+			}
+		}
+	}
+}
+#endif
\ No newline at end of file
diff --git a/Engine/Plugins/Experimental/DaySequence/Source/DaySequence/Private/DaySequenceModifierComponent.cpp b/Engine/Plugins/Experimental/DaySequence/Source/DaySequence/Private/DaySequenceModifierComponent.cpp
--- a/Engine/Plugins/Experimental/DaySequence/Source/DaySequence/Private/DaySequenceModifierComponent.cpp
+++ b/Engine/Plugins/Experimental/DaySequence/Source/DaySequence/Private/DaySequenceModifierComponent.cpp
@@ -1338,7 +1338,7 @@ bool UDaySequenceModifierComponent::GetBlendPosition(FVector& InPosition) const
 	CSV_SCOPED_TIMING_STAT_EXCLUSIVE(DaySequenceModifier_GetBlendPosition);
 	
 #if WITH_EDITOR
-	if (!GetWorld()->IsGameWorld())
+	if (const UWorld* World = GetWorld(); World && !World->IsGameWorld())
 	{
 		InPosition = UE::DaySequence::GVolumePreviewLocation;
 		return true;
diff --git a/Engine/Plugins/Experimental/DaySequence/Source/DaySequence/Public/DaySequenceActor.h b/Engine/Plugins/Experimental/DaySequence/Source/DaySequence/Public/DaySequenceActor.h
--- a/Engine/Plugins/Experimental/DaySequence/Source/DaySequence/Public/DaySequenceActor.h
+++ b/Engine/Plugins/Experimental/DaySequence/Source/DaySequence/Public/DaySequenceActor.h
@@ -323,6 +323,10 @@ public:
 
 	/* Called to register InFunction as a callback for instances matching InConditions with a lifetime equivalent to LifetimeObject. */
 	void BindToConditionCallbacks(UObject* LifetimeObject, const FDaySequenceConditionSet::FConditionValueMap& InConditions, const TFunction<void(void)>& InFunction);
+
+#if WITH_EDITOR
+	void HandleConditionReinstanced(const FCoreUObjectDelegates::FReplacementObjectMap& OldToNewInstanceMap);
+#endif
 	
 protected:
 	//~ Begin UObject interface
