DaySequence: Fix modifiers not enabling properly when the day cycle is paused.

Also fixes an issue with non-blending volumes not working when a lower bias subsequence animates some properties also animated by the non-blending volume

Also enables blending for volumes when the day cycle is paused. Note: this may be prone to hitching because we force evaluations that are synchronous. The amortized cost should not be higher than when normal playback results in asynchronous evaluations, but these types of evaluations do all of the work in a single frame.

diff --git a/Engine/Plugins/Experimental/DaySequence/Source/DaySequence/Private/DaySequenceActor.cpp b/Engine/Plugins/Experimental/DaySequence/Source/DaySequence/Private/DaySequenceActor.cpp
--- a/Engine/Plugins/Experimental/DaySequence/Source/DaySequence/Private/DaySequenceActor.cpp
+++ b/Engine/Plugins/Experimental/DaySequence/Source/DaySequence/Private/DaySequenceActor.cpp
@@ -348,6 +348,7 @@ void ADaySequenceActor::EndPlay(const EEndPlayReason::Type EndPlayReason)
 
 		Player->OnPlay.RemoveAll(this);
 		Player->OnPlayReverse.RemoveAll(this);
+		Player->OnPause.RemoveAll(this);
 		Player->OnStop.RemoveAll(this);
 
 		Player->TearDown();
@@ -517,12 +518,18 @@ void ADaySequenceActor::InitializePlayer()
 		}
 
 		UDaySequencePlayer* Player = GetSequencePlayer();
-		if (OldTimeControllerOverride && Player && OldTimeControllerOverride != Player->GetTimeController())
+		if (Player)
 		{
-			ensure(OldTimeControllerOverride->PreviousTimeController == nullptr);
+			if (OldTimeControllerOverride && OldTimeControllerOverride != Player->GetTimeController())
+			{
+				ensure(OldTimeControllerOverride->PreviousTimeController == nullptr);
+
+				OldTimeControllerOverride->PreviousTimeController = Player->GetTimeController();
+				Player->SetTimeControllerDirectly(OldTimeControllerOverride);
+			}
 
-			OldTimeControllerOverride->PreviousTimeController = Player->GetTimeController();
-			Player->SetTimeControllerDirectly(OldTimeControllerOverride);
+			Player->OnPlay.AddUniqueDynamic(this, &ADaySequenceActor::StopDaySequenceUpdateTimer);
+			Player->OnPause.AddUniqueDynamic(this, &ADaySequenceActor::StartDaySequenceUpdateTimer);
 		}
 	}
 }
@@ -854,13 +861,46 @@ void ADaySequenceActor::OnSequencePlayerUpdate(const UMovieSceneSequencePlayer&
 	const float PreviousHours = FrameTimeToDayHours(PreviousTime);
 	SequencePlayerUpdated(CurrentHours, PreviousHours);
 	
-	OnSequencePlayerUpdated.Broadcast();
+	OnDaySequenceUpdate.Broadcast();
 }
 
 void ADaySequenceActor::SequencePlayerUpdated(float CurrentTime, float PreviousTime)
 {
 }
 
+void ADaySequenceActor::StartDaySequenceUpdateTimer()
+{
+	if (HasAuthority())
+	{
+		return;
+	}
+
+	if (const UWorld* World = GetWorld())
+	{
+		FTimerManagerTimerParameters TimerParameters;
+		TimerParameters.bLoop = true;
+		TimerParameters.bMaxOncePerFrame = true;
+
+		World->GetTimerManager().SetTimer(DaySequenceUpdateTimerHandle, [this]()
+		{
+			OnDaySequenceUpdate.Broadcast();
+		}, PrimaryActorTick.TickInterval, TimerParameters);
+	}
+}
+
+void ADaySequenceActor::StopDaySequenceUpdateTimer()
+{
+	if (HasAuthority())
+	{
+		return;
+	}
+
+	if (const UWorld* World = GetWorld())
+	{
+		World->GetTimerManager().ClearTimer(DaySequenceUpdateTimerHandle);
+	}
+}
+
 #if WITH_EDITORONLY_DATA
 bool ADaySequenceActor::GetOverrideInitialTimeOfDay() const
 {
diff --git a/Engine/Plugins/Experimental/DaySequence/Source/DaySequence/Private/DaySequenceModifierComponent.cpp b/Engine/Plugins/Experimental/DaySequence/Source/DaySequence/Private/DaySequenceModifierComponent.cpp
--- a/Engine/Plugins/Experimental/DaySequence/Source/DaySequence/Private/DaySequenceModifierComponent.cpp
+++ b/Engine/Plugins/Experimental/DaySequence/Source/DaySequence/Private/DaySequenceModifierComponent.cpp
@@ -404,7 +404,7 @@ void UDaySequenceModifierComponent::OnUnregister()
 	RemoveSubSequenceTrack();
 }
 
-void UDaySequenceModifierComponent::SequencePlayerUpdated()
+void UDaySequenceModifierComponent::DaySequenceUpdate()
 {
 	CSV_SCOPED_TIMING_STAT_EXCLUSIVE(DaySequenceModifier_SequencePlayerUpdated);
 	
@@ -479,7 +479,7 @@ void UDaySequenceModifierComponent::BindToDaySequenceActor(ADaySequenceActor* Da
 	if (ensureMsgf(DaySequenceActor, TEXT("BindToDaySequenceActor called with a null Day Sequence Actor.")))
 	{
 		DaySequenceActor->GetOnPostInitializeDaySequences().AddUObject(this, &UDaySequenceModifierComponent::ReinitializeSubSequence);
-		DaySequenceActor->GetOnSequencePlayerUpdated().AddUObject(this, &UDaySequenceModifierComponent::SequencePlayerUpdated);
+		DaySequenceActor->GetOnDaySequenceUpdate().AddUObject(this, &UDaySequenceModifierComponent::DaySequenceUpdate);
 #if ENABLE_DRAW_DEBUG
 		if (!DaySequenceActor->IsDebugCategoryRegistered(ShowDebug_ModifierCategory))
 		{
@@ -500,7 +500,7 @@ void UDaySequenceModifierComponent::UnbindFromDaySequenceActor()
 	if (TargetActor)
 	{
 		TargetActor->GetOnPostInitializeDaySequences().RemoveAll(this);
-		TargetActor->GetOnSequencePlayerUpdated().RemoveAll(this);
+		TargetActor->GetOnDaySequenceUpdate().RemoveAll(this);
 #if ENABLE_DRAW_DEBUG
 		TargetActor->GetOnDebugLevelChanged().RemoveAll(this);
 		TargetActor->UnregisterDebugEntry(DebugEntry, ShowDebug_ModifierCategory);
@@ -826,7 +826,9 @@ UMovieSceneSubSection* UDaySequenceModifierComponent::InitializeDaySequence(cons
 		SubSection->ClearFlags(RF_Transactional);
 		// SubSections of DaySequenceTrack will inherit flags from its parent track - RF_Transient in this case.
 		SubSection->Parameters.HierarchicalBias = Bias + BiasOffset;
-		SubSection->Parameters.Flags            = EMovieSceneSubSectionFlags::OverrideRestoreState | (bIgnoreBias ? EMovieSceneSubSectionFlags::IgnoreHierarchicalBias : EMovieSceneSubSectionFlags::None);
+		SubSection->Parameters.Flags            = EMovieSceneSubSectionFlags::OverrideRestoreState
+												| (bIgnoreBias ? EMovieSceneSubSectionFlags::IgnoreHierarchicalBias : EMovieSceneSubSectionFlags::None)
+												| (bBlendHierarchicalBias ? EMovieSceneSubSectionFlags::BlendHierarchicalBias : EMovieSceneSubSectionFlags::None);
 
 		SubSection->SetSequence(Sequence);
 		SubSection->SetRange(MovieScene->GetPlaybackRange());
@@ -847,7 +849,6 @@ UMovieSceneSubSection* UDaySequenceModifierComponent::InitializeDaySequence(cons
 
 			EasingFunction->Initialize(UDaySequenceModifierEasingFunction::EEasingFunctionType::EaseOut);
 			SubSection->Easing.EaseOut = EasingFunction;
-			SubSection->Parameters.Flags |= EMovieSceneSubSectionFlags::BlendHierarchicalBias;
 		}
 
 #if WITH_EDITOR
@@ -1433,9 +1434,18 @@ float UDaySequenceModifierComponent::GetCurrentBlendWeight() const
 
 float UDaySequenceModifierComponent::UpdateBlendWeight() const
 {
-	FVector Position;
-	const float DistanceBlendFactor = GetBlendPosition(Position) ? GetDistanceBlendFactor(Position) : 1.f;
-	return FMath::Min(DistanceBlendFactor, CustomVolumeBlendWeight);
+	FVector BlendPosition;
+	const bool bHasBlendPosition = GetBlendPosition(BlendPosition);
+
+	const float OldBlendWeight = CachedDistanceBlendFactor;
+	const float NewBlendWeight = FMath::Min(bHasBlendPosition ? GetDistanceBlendFactor(BlendPosition) : 1.f, CustomVolumeBlendWeight);
+
+	if (bHasBlendPosition && TargetActor && !TargetActor->IsPlaying() && !FMath::IsNearlyEqual(OldBlendWeight, NewBlendWeight))
+	{
+		TargetActor->SetTimeOfDay(TargetActor->GetTimeOfDay());
+	}
+	
+	return NewBlendWeight;
 }
 
 void UDaySequenceModifierComponent::SetVolumeCollisionEnabled(const ECollisionEnabled::Type InCollisionType) const
diff --git a/Engine/Plugins/Experimental/DaySequence/Source/DaySequence/Public/DaySequenceActor.h b/Engine/Plugins/Experimental/DaySequence/Source/DaySequence/Public/DaySequenceActor.h
--- a/Engine/Plugins/Experimental/DaySequence/Source/DaySequence/Public/DaySequenceActor.h
+++ b/Engine/Plugins/Experimental/DaySequence/Source/DaySequence/Public/DaySequenceActor.h
@@ -9,6 +9,8 @@
 #include "MovieSceneBindingOwnerInterface.h"
 
 #include "Misc/Timecode.h"
+#include "TimerManager.h"
+
 #include "DaySequenceActor.generated.h"
 
 class FDebugDisplayInfo;
@@ -281,8 +283,13 @@ public:
 	DECLARE_EVENT_OneParam(ADaySequenceActor, FOnPostInitializeDaySequences, FSubSectionPreserveMap*)
 	FOnPostInitializeDaySequences& GetOnPostInitializeDaySequences() { return OnPostInitializeDaySequences; }
 
-	DECLARE_EVENT(ADaySequenceActor, FOnSequencePlayerUpdated)
-	FOnSequencePlayerUpdated& GetOnSequencePlayerUpdated() { return OnSequencePlayerUpdated; }
+	/**
+	 * This delegate is broadcast at a rate matching this actor's tick interval.
+	 * It is either broadcast after each sequence player update or by a timer that is configured to run only when the sequence player is paused.
+	 * Used to synchronize polling logic that is not owned by this actor when it should occur at an interval specified by this actor.
+	 */
+	DECLARE_EVENT(ADaySequenceActor, FOnDaySequenceUpdate)
+	FOnDaySequenceUpdate& GetOnDaySequenceUpdate() { return OnDaySequenceUpdate; }
 	
 	void InvalidateMuteStates() const;
 	
@@ -522,8 +529,18 @@ protected:
 	FOnRootSequenceChanged OnPreRootSequenceChanged;
 	FOnRootSequenceChanged OnPostRootSequenceChanged;
 	FOnPostInitializeDaySequences OnPostInitializeDaySequences;
-	FOnSequencePlayerUpdated OnSequencePlayerUpdated;
 	UE::DaySequence::FOnInvalidateMuteStates OnInvalidateMuteStates;
+
+	FOnDaySequenceUpdate OnDaySequenceUpdate;
+	FTimerHandle DaySequenceUpdateTimerHandle;
+
+	/** Starts a timer that will broadcast OnDaySequenceUpdate when the sequence player is paused. */
+	UFUNCTION()
+	void StartDaySequenceUpdateTimer();
+
+	/** Stops the timer that was started by DaySequenceUpdateTimerHandle. */
+	UFUNCTION()
+	void StopDaySequenceUpdateTimer();
 	
 #if ENABLE_DRAW_DEBUG
 	void OnShowDebugInfo(AHUD* HUD, UCanvas* Canvas, const FDebugDisplayInfo& DisplayInfo, float& YL, float& YPos);
diff --git a/Engine/Plugins/Experimental/DaySequence/Source/DaySequence/Public/DaySequenceModifierComponent.h b/Engine/Plugins/Experimental/DaySequence/Source/DaySequence/Public/DaySequenceModifierComponent.h
--- a/Engine/Plugins/Experimental/DaySequence/Source/DaySequence/Public/DaySequenceModifierComponent.h
+++ b/Engine/Plugins/Experimental/DaySequence/Source/DaySequence/Public/DaySequenceModifierComponent.h
@@ -329,7 +329,8 @@ public:
 	void OnUnregister() override;
 	/*~ End UActorComponent interface */
 
-	void SequencePlayerUpdated();
+	/** Bound to delegate on the DaySequenceActor that allows all modifiers to do work at appropriate times at the specific actors tick interval. */
+	void DaySequenceUpdate();
 	
 #if WITH_EDITOR
 	/*~ Begin FTickableGameObject interface */
