SunPositionDaySequenceActor: Add moon, change some property defaults

Replacing BP_SunAndMoon_DSA with the native parent actor. This entails:
- Adding moon (with modified defaults for intensity, forward shading priority, and temperature)
- Changing sky material
- Adding a default collection asset (to animate the moon and sky material)

diff --git a/Engine/Plugins/Experimental/DaySequence/Source/DaySequence/Private/Actors/SunPositionDaySequenceActor.cpp b/Engine/Plugins/Experimental/DaySequence/Source/DaySequence/Private/Actors/SunPositionDaySequenceActor.cpp
--- a/Engine/Plugins/Experimental/DaySequence/Source/DaySequence/Private/Actors/SunPositionDaySequenceActor.cpp
+++ b/Engine/Plugins/Experimental/DaySequence/Source/DaySequence/Private/Actors/SunPositionDaySequenceActor.cpp
@@ -3,8 +3,13 @@
 #include "Actors/SunPositionDaySequenceActor.h"
 
 #include "Components/DirectionalLightComponent.h"
+#include "Components/StaticMeshComponent.h"
+#include "DaySequenceCollectionAsset.h"
 #include "ProceduralDaySequenceBuilder.h"
 
+#include "Materials/MaterialInterface.h"
+#include "UObject/ConstructorHelpers.h"
+
 namespace UE::DaySequence
 {
 	struct FSunPositionData
@@ -164,6 +169,31 @@ namespace UE::DaySequence
 	}
 }
 
+ASunPositionDaySequenceActor::ASunPositionDaySequenceActor(const FObjectInitializer& Init)
+: Super(Init)
+{
+	MoonComponent = CreateDefaultSubobject<UDirectionalLightComponent>(TEXT("Moon"));
+	MoonComponent->SetupAttachment(SunRootComponent);
+
+	// Give Sun forward shading priority.
+	SunComponent->SetForwardShadingPriority(1);
+	MoonComponent->SetForwardShadingPriority(0);
+
+	// Configure other Moon defaults.
+	MoonComponent->SetAtmosphereSunLightIndex(1);	// Make Moon the secondary directional light that contributes to the sky atmosphere.
+	MoonComponent->SetIntensity(2.f);
+	MoonComponent->SetUseTemperature(true);
+	MoonComponent->SetTemperature(9000.f);
+	
+	// Override the sky sphere material.
+	static ConstructorHelpers::FObjectFinder<UMaterialInterface> SkySphereDefaultMaterial(TEXT("/DaySequence/MI_24hrSky.MI_24hrSky"));
+	SkySphereComponent->SetMaterial(0, SkySphereDefaultMaterial.Object.Get());
+
+	// Override the default collection (which animates the moon and sky material)
+	static ConstructorHelpers::FObjectFinder<UDaySequenceCollectionAsset> DefaultCollection(TEXT("/DaySequence/DSCA_24hr.DSCA_24hr"));
+	DaySequenceCollection = DefaultCollection.Object.Get();
+}
+
 void ASunPositionDaySequenceActor::Tick(float DeltaTime)
 {
 	Super::Tick(DeltaTime);
diff --git a/Engine/Plugins/Experimental/DaySequence/Source/DaySequence/Public/Actors/SunPositionDaySequenceActor.h b/Engine/Plugins/Experimental/DaySequence/Source/DaySequence/Public/Actors/SunPositionDaySequenceActor.h
--- a/Engine/Plugins/Experimental/DaySequence/Source/DaySequence/Public/Actors/SunPositionDaySequenceActor.h
+++ b/Engine/Plugins/Experimental/DaySequence/Source/DaySequence/Public/Actors/SunPositionDaySequenceActor.h
@@ -16,6 +16,9 @@ class DAYSEQUENCE_API ASunPositionDaySequenceActor
 {
 	GENERATED_BODY()
 
+public:
+	ASunPositionDaySequenceActor(const FObjectInitializer& Init);
+
 protected:
 	virtual void Tick(float DeltaTime) override;
 
@@ -27,6 +30,10 @@ protected:
 #if WITH_EDITOR
 	virtual void PostEditChangeProperty(FPropertyChangedEvent& PropertyChangedEvent) override;
 #endif
+
+protected:
+	UPROPERTY(VisibleDefaultsOnly, BlueprintReadOnly, Category= "Day Sequence", meta = (AllowPrivateAccess = "true"))
+	TObjectPtr<UDirectionalLightComponent> MoonComponent;
 	
 private:
 	
