DaySequenceActor: Early out in InitializeRootSequence for CDO

I am not sure how problematic this is, but I discovered that the RootSequence was being generated initially on loading CDO and it was being copied over to instances. This meant that the RootSequence was outered to the CDO.

diff --git a/Engine/Plugins/Experimental/DaySequence/Source/DaySequence/Private/DaySequenceActor.cpp b/Engine/Plugins/Experimental/DaySequence/Source/DaySequence/Private/DaySequenceActor.cpp
--- a/Engine/Plugins/Experimental/DaySequence/Source/DaySequence/Private/DaySequenceActor.cpp
+++ b/Engine/Plugins/Experimental/DaySequence/Source/DaySequence/Private/DaySequenceActor.cpp
@@ -539,6 +539,11 @@ void ADaySequenceActor::InitializePlayer()
 void ADaySequenceActor::InitializeRootSequence()
 {
 	CSV_SCOPED_TIMING_STAT_EXCLUSIVE(DaySequenceActor_InitializeRootSequence);
+
+	if (IsTemplate())
+	{
+		return;
+	}
 	
 	ensureMsgf(!SequencePlayer || !SequencePlayer->IsValid(), TEXT("InitializeRootSequence called but the sequence player has already been initialized."));
 	
