DaySequence: Replace FTimecode with FDaySequenceTime

#changelist validated
#virtualized

diff --git a/Engine/Plugins/Experimental/DaySequence/Source/DaySequence/Private/DaySequenceActor.cpp b/Engine/Plugins/Experimental/DaySequence/Source/DaySequence/Private/DaySequenceActor.cpp
--- a/Engine/Plugins/Experimental/DaySequence/Source/DaySequence/Private/DaySequenceActor.cpp
+++ b/Engine/Plugins/Experimental/DaySequence/Source/DaySequence/Private/DaySequenceActor.cpp
@@ -71,7 +71,7 @@ namespace UE::DaySequence
 		float TimeInHours;
 		FFrameTime StaticTime;
 #if WITH_EDITORONLY_DATA
-		FTimecode PreviousTimeOfDayPreview;
+		FDaySequenceTime PreviousTimeOfDayPreview;
 #endif
 
 		FStaticTimeControllerOverride(FFrameTime InStaticTime, TSharedPtr<FMovieSceneTimeController> InPreviousTimeController)
@@ -112,9 +112,9 @@ ADaySequenceActor::ADaySequenceActor(const FObjectInitializer& Init)
 : Super(Init)
 , bRunDayCycle(true)
 , bUseInterpCurve(false)
-, DayLength(24, 0, 0, 0, false)
-, TimePerCycle(0, 5, 0, 0, false)
-, InitialTimeOfDay(6, 0, 0, 0, false)
+, DayLength(24, 0, 0)
+, TimePerCycle(0, 5, 0)
+, InitialTimeOfDay(6, 0, 0)
 {
 	USceneComponent* SceneRootComponent = CreateDefaultSubobject<USceneComponent>(USceneComponent::GetDefaultSceneRootVariableName());
 	SetRootComponent(SceneRootComponent);
@@ -140,7 +140,7 @@ ADaySequenceActor::ADaySequenceActor(const FObjectInitializer& Init)
 		}
 	}
 
-	TimeOfDayPreview = FTimecode(6, 0, 0, 0, false);
+	TimeOfDayPreview = FDaySequenceTime(6, 0, 0);
 #endif // WITH_EDITORONLY_DATA
 
 #if WITH_EDITOR || ENABLE_DRAW_DEBUG
@@ -468,7 +468,7 @@ void ADaySequenceActor::ConditionalSetTimeOfDayPreview(float InHours)
 	// Wrap the input hours using day length
 	InHours = FMath::Frac(InHours / GetDayLength()) * GetDayLength(); 
 	const TOptional<FFrameRate> FrameRate = RootSequence ? RootSequence->GetMovieScene()->GetDisplayRate() : TOptional<FFrameRate>();
-	const FTimecode NewTimeOfDayPreview = DaySequenceHelpers::HoursToTimecode(InHours, FrameRate);
+	const FDaySequenceTime NewTimeOfDayPreview = DaySequenceHelpers::HoursToTimecode(InHours, FrameRate);
 	if (NewTimeOfDayPreview != TimeOfDayPreview)
 	{
 		TimeOfDayPreview = NewTimeOfDayPreview;
@@ -495,7 +495,7 @@ void ADaySequenceActor::SetTimeOfDayPreview(float InHours)
 	// Wrap the input hours using day length
 	InHours = FMath::Frac(InHours / GetDayLength()) * GetDayLength(); 
 	const TOptional<FFrameRate> FrameRate = RootSequence ? RootSequence->GetMovieScene()->GetDisplayRate() : TOptional<FFrameRate>();
-	const FTimecode NewTimeOfDayPreview = DaySequenceHelpers::HoursToTimecode(InHours, FrameRate);
+	const FDaySequenceTime NewTimeOfDayPreview = DaySequenceHelpers::HoursToTimecode(InHours, FrameRate);
 	TimeOfDayPreview = NewTimeOfDayPreview;
 	OnTimeOfDayPreviewChangedEvent.Broadcast(InHours);
 	OnTimeOfDayPreviewChanged.Broadcast(InHours);
@@ -1039,7 +1039,7 @@ void ADaySequenceActor::SetDayLength(float InHours)
 	const TOptional<FFrameRate> FrameRate = RootSequence ? RootSequence->GetMovieScene()->GetDisplayRate() : TOptional<FFrameRate>();
 
 	// Set min day length to 1 second.
-	const FTimecode OneSecond(0, 0, 1, 0, false);
+	const FDaySequenceTime OneSecond(0, 0, 1);
 	InHours = FMath::Max(InHours, DaySequenceHelpers::TimecodeToHours(OneSecond, FrameRate));
 	DayLength = DaySequenceHelpers::HoursToTimecode(InHours, FrameRate);
 }
@@ -1055,7 +1055,7 @@ void ADaySequenceActor::SetTimePerCycle(float InHours)
 	const TOptional<FFrameRate> FrameRate = RootSequence ? RootSequence->GetMovieScene()->GetDisplayRate() : TOptional<FFrameRate>();
 
 	// Set min cycle length to 1 second.
-	const FTimecode OneSecond(0, 0, 1, 0, false);
+	const FDaySequenceTime OneSecond(0, 0, 1);
 	InHours = FMath::Max(InHours, DaySequenceHelpers::TimecodeToHours(OneSecond, FrameRate));
 	TimePerCycle = DaySequenceHelpers::HoursToTimecode(InHours, FrameRate);
 }
@@ -1068,16 +1068,16 @@ void ADaySequenceActor::Multicast_SetTimePerCycle_Implementation(float InHours)
 	const float CurrentTimeOfDay = bIsGameWorld ? GetTimeOfDay() : 0.f;
 
 	// Set min cycle length to 1 second.
-	const FTimecode OneSecond(0, 0, 1, 0, false);
+	const FDaySequenceTime OneSecond(0, 0, 1);
 	InHours = FMath::Max(InHours, DaySequenceHelpers::TimecodeToHours(OneSecond, FrameRate));
-	const FTimecode NewTimePerCycle = DaySequenceHelpers::HoursToTimecode(InHours, FrameRate); 
+	const FDaySequenceTime NewTimePerCycle = DaySequenceHelpers::HoursToTimecode(InHours, FrameRate); 
 	if (NewTimePerCycle != TimePerCycle)
 	{
 		// Validate the new TimePerCycle to avoid overflowing the playback range.
 		UMovieScene* RootMovieScene = RootSequence->GetMovieScene();
 		const FFrameRate DisplayRate = RootMovieScene->GetDisplayRate();
 		const FFrameRate TickResolution = RootMovieScene->GetTickResolution();
-		const auto IsTimePerCycleOverflow = [&DisplayRate, &TickResolution](const FTimecode& InTimePerCycle)
+		const auto IsTimePerCycleOverflow = [&DisplayRate, &TickResolution](const FDaySequenceTime& InTimePerCycle)
 		{
 			const float TimePerCycleSeconds = DaySequenceHelpers::TimecodeToSeconds(InTimePerCycle, TOptional<FFrameRate>());
 			const int32 CycleSeconds = static_cast<int32>(TimePerCycleSeconds) + 1;
@@ -1098,8 +1098,7 @@ void ADaySequenceActor::Multicast_SetTimePerCycle_Implementation(float InHours)
 		if (bIsGameWorld)
 		{
 			// Update playback range for the root sequence.
-			const FQualifiedFrameTime FrameTimePerCycle(TimePerCycle, FFrameRate());
-			const float DaySeconds = FrameTimePerCycle.AsSeconds();
+			const float DaySeconds = DaySequenceHelpers::TimecodeToSeconds(TimePerCycle, FFrameRate());
 			const int32 RootDuration = RootMovieScene->GetTickResolution().AsFrameNumber(DaySeconds).Value;
 			RootMovieScene->MarkAsChanged();
 			RootMovieScene->SetPlaybackRange(0, RootDuration);
diff --git a/Engine/Plugins/Experimental/DaySequence/Source/DaySequence/Private/DaySequenceHelpers.cpp b/Engine/Plugins/Experimental/DaySequence/Source/DaySequence/Private/DaySequenceHelpers.cpp
--- a/Engine/Plugins/Experimental/DaySequence/Source/DaySequence/Private/DaySequenceHelpers.cpp
+++ b/Engine/Plugins/Experimental/DaySequence/Source/DaySequence/Private/DaySequenceHelpers.cpp
@@ -1,55 +1,42 @@
 // Copyright Epic Games, Inc. All Rights Reserved.
 
 #include "DaySequenceHelpers.h"
+
 #include "Misc/FrameRate.h"
 #include "Misc/QualifiedFrameTime.h"
 
-static constexpr float SecondsPerHour = 3600.f;
-static constexpr float SecondsPerMinute = 60.f;
-
-float DaySequenceHelpers::TimecodeToHours(const FTimecode& InTimecode, const TOptional<FFrameRate>& InRate)
+float DaySequenceHelpers::TimecodeToHours(const FDaySequenceTime& InTimecode, const TOptional<FFrameRate>& InRate)
 {
-	return TimecodeToSeconds(InTimecode, InRate) / SecondsPerHour;
+	return TimecodeToSeconds(InTimecode, InRate) / FDaySequenceTime::SecondsPerHour;
 }
 
-float DaySequenceHelpers::TimecodeToSeconds(const FTimecode& InTimecode, const TOptional<FFrameRate>& InRate)
+float DaySequenceHelpers::TimecodeToSeconds(const FDaySequenceTime& InTimecode, const TOptional<FFrameRate>& InRate)
 {
 	float Result;
-	if (InRate.IsSet())
-	{
-		const FQualifiedFrameTime FrameTime(InTimecode, *InRate);
-		Result = FrameTime.AsSeconds();
-	}
-	else
-	{
-		Result = InTimecode.Hours * SecondsPerHour + InTimecode.Minutes * SecondsPerMinute + InTimecode.Seconds; 
-	}
+	
+	Result = InTimecode.Hours * FDaySequenceTime::SecondsPerHour + InTimecode.Minutes * FDaySequenceTime::SecondsPerMinute + InTimecode.Seconds; 
+
 	return Result;
 }
 
-FTimecode DaySequenceHelpers::HoursToTimecode(float InHours, const TOptional<FFrameRate>& InRate)
+FDaySequenceTime DaySequenceHelpers::HoursToTimecode(float InHours, const TOptional<FFrameRate>& InRate)
 {
-	return SecondsToTimecode(InHours * SecondsPerHour, InRate);
+	return SecondsToTimecode(InHours * FDaySequenceTime::SecondsPerHour, InRate);
 }
 
-FTimecode DaySequenceHelpers::SecondsToTimecode(float InSeconds, const TOptional<FFrameRate>& InRate)
+FDaySequenceTime DaySequenceHelpers::SecondsToTimecode(float InSeconds, const TOptional<FFrameRate>& InRate)
 {
-	FTimecode Result;
-	if (InRate.IsSet())
-	{
-		Result = FQualifiedFrameTime(InRate->AsFrameTime(InSeconds), *InRate).ToTimecode();
-	}
-	else
-	{
-		float Seconds = InSeconds;
-		const float Hours = FMath::Floor(InSeconds / SecondsPerHour);
-		Seconds -= Hours * SecondsPerHour;
-		const float Minutes = FMath::Floor(Seconds / SecondsPerMinute);
-		Seconds -= Minutes * SecondsPerMinute;
-		Result.Hours = Hours;
-		Result.Minutes = Minutes;
-		Result.Seconds = FMath::Floor(Seconds);
-	}
+	FDaySequenceTime Result;
+
+	float Seconds = InSeconds;
+	const float Hours = FMath::Floor(InSeconds / FDaySequenceTime::SecondsPerHour);
+	Seconds -= Hours * FDaySequenceTime::SecondsPerHour;
+	const float Minutes = FMath::Floor(Seconds / FDaySequenceTime::SecondsPerMinute);
+	Seconds -= Minutes * FDaySequenceTime::SecondsPerMinute;
+	Result.Hours = Hours;
+	Result.Minutes = Minutes;
+	Result.Seconds = FMath::Floor(Seconds);
+	
 	return Result;
 }
 
diff --git a/Engine/Plugins/Experimental/DaySequence/Source/DaySequence/Private/DaySequenceHelpers.h b/Engine/Plugins/Experimental/DaySequence/Source/DaySequence/Private/DaySequenceHelpers.h
--- a/Engine/Plugins/Experimental/DaySequence/Source/DaySequence/Private/DaySequenceHelpers.h
+++ b/Engine/Plugins/Experimental/DaySequence/Source/DaySequence/Private/DaySequenceHelpers.h
@@ -2,6 +2,7 @@
 
 #pragma once
 
+#include "DaySequenceTime.h"
 #include "Misc/OptionalFwd.h"
 
 struct FTimecode;
@@ -20,7 +21,7 @@ public:
 	 * @param InRate Optional frame rate to convert frames portion of Timecode.
 	 * @return result in hours.
 	 */
-	static float TimecodeToHours(const FTimecode& InTimecode, const TOptional<FFrameRate>& InRate);
+	static float TimecodeToHours(const FDaySequenceTime& InTimecode, const TOptional<FFrameRate>& InRate);
 
 	/**
 	 * Convert Timecode to seconds. Optionally provide a frame rate to convert
@@ -31,7 +32,7 @@ public:
 	 * @param InRate Optional frame rate to convert frames portion of Timecode.
 	 * @return result in seconds.
 	 */
-	static float TimecodeToSeconds(const FTimecode& InTimecode, const TOptional<FFrameRate>& InRate);
+	static float TimecodeToSeconds(const FDaySequenceTime& InTimecode, const TOptional<FFrameRate>& InRate);
 
 	/**
 	 * Convert hours to Timecode. Optionally provide a frame rate to convert
@@ -42,7 +43,7 @@ public:
 	 * @param InRate Optional frame rate to convert fractional seconds to frames.
 	 * @return result in timecode.
 	 */
-	static FTimecode HoursToTimecode(float InHours, const TOptional<FFrameRate>& InRate);
+	static FDaySequenceTime HoursToTimecode(float InHours, const TOptional<FFrameRate>& InRate);
 
 	/**
 	 * Convert seconds to Timecode. Optionally provide a frame rate to convert
@@ -53,5 +54,5 @@ public:
 	 * @param InRate Optional frame rate to convert fractional seconds to frames.
 	 * @return result in timecode.
 	 */
-	static FTimecode SecondsToTimecode(float InSeconds, const TOptional<FFrameRate>& InRate);
+	static FDaySequenceTime SecondsToTimecode(float InSeconds, const TOptional<FFrameRate>& InRate);
 };
diff --git a/Engine/Plugins/Experimental/DaySequence/Source/DaySequence/Private/DaySequenceTime.cpp b/Engine/Plugins/Experimental/DaySequence/Source/DaySequence/Private/DaySequenceTime.cpp
new file mode 100644
--- /dev/null
+++ b/Engine/Plugins/Experimental/DaySequence/Source/DaySequence/Private/DaySequenceTime.cpp
@@ -0,0 +1,45 @@
+// Copyright Epic Games, Inc. All Rights Reserved.
+
+#include "DaySequenceTime.h"
+
+FDaySequenceTime::FDaySequenceTime(int32 InHours, int32 InMinutes, int32 InSeconds)
+: Hours(InHours)
+, Minutes(InMinutes)
+, Seconds(InSeconds)
+{}
+
+bool FDaySequenceTime::operator==(const FDaySequenceTime& RHS) const
+{
+	return Hours == RHS.Hours && Minutes == RHS.Minutes && Seconds == RHS.Seconds;
+}
+
+bool FDaySequenceTime::operator!=(const FDaySequenceTime& RHS) const
+{
+	return !operator==(RHS);
+}
+
+float FDaySequenceTime::ToHours() const
+{
+	return ToSeconds() / SecondsPerHour;
+}
+
+float FDaySequenceTime::ToSeconds() const
+{
+	return Hours * SecondsPerHour + Minutes * SecondsPerMinute + Seconds;
+}
+
+FDaySequenceTime FDaySequenceTime::FromHours(float InHours)
+{
+	return FromSeconds(InHours * SecondsPerHour);
+}
+
+FDaySequenceTime FDaySequenceTime::FromSeconds(float InSeconds)
+{
+	float Seconds = InSeconds;
+	const float Hours = FMath::Floor(InSeconds / SecondsPerHour);
+	Seconds -= Hours * SecondsPerHour;
+	const float Minutes = FMath::Floor(Seconds / SecondsPerMinute);
+	Seconds -= Minutes * SecondsPerMinute;
+
+	return FDaySequenceTime(Hours, Minutes, FMath::Floor(Seconds));
+}
\ No newline at end of file
diff --git a/Engine/Plugins/Experimental/DaySequence/Source/DaySequence/Private/ProceduralDaySequenceBuilder.cpp b/Engine/Plugins/Experimental/DaySequence/Source/DaySequence/Private/ProceduralDaySequenceBuilder.cpp
--- a/Engine/Plugins/Experimental/DaySequence/Source/DaySequence/Private/ProceduralDaySequenceBuilder.cpp
+++ b/Engine/Plugins/Experimental/DaySequence/Source/DaySequence/Private/ProceduralDaySequenceBuilder.cpp
@@ -2,7 +2,7 @@
 
 #include "ProceduralDaySequenceBuilder.h"
 
-#include "DaySequenceHelpers.h"
+#include "DaySequenceTime.h"
 #include "MovieSceneCommonHelpers.h"
 
 #include "Sections/MovieSceneBoolSection.h"
@@ -76,9 +76,9 @@ UDaySequence* UProceduralDaySequenceBuilder::Initialize(ADaySequenceActor* InAct
 		const FName SequenceName = MakeUniqueObjectName(InActor, UDaySequence::StaticClass());
 		ProceduralDaySequence = NewObject<UDaySequence>(InActor, SequenceName, RF_Transient);
 		ProceduralDaySequence->Initialize(RF_Transient);
-		const TOptional<FFrameRate> FrameRate = TargetActor->GetRootSequence() ? TargetActor->GetRootSequence()->GetMovieScene()->GetDisplayRate() : TOptional<FFrameRate>();
-		const FQualifiedFrameTime FrameTimePerCycle(DaySequenceHelpers::HoursToTimecode(TargetActor->GetTimePerCycle(), FrameRate), FFrameRate());
-		const float DaySeconds = FrameTimePerCycle.AsSeconds();
+
+		const float DaySeconds = TargetActor->GetTimePerCycle() * FDaySequenceTime::SecondsPerHour;
+		
 		const int32 Duration = ProceduralDaySequence->GetMovieScene()->GetTickResolution().AsFrameNumber(DaySeconds).Value;
 		ProceduralDaySequence->GetMovieScene()->SetPlaybackRange(0, Duration);
 	}
diff --git a/Engine/Plugins/Experimental/DaySequence/Source/DaySequence/Public/DaySequenceActor.h b/Engine/Plugins/Experimental/DaySequence/Source/DaySequence/Public/DaySequenceActor.h
--- a/Engine/Plugins/Experimental/DaySequence/Source/DaySequence/Public/DaySequenceActor.h
+++ b/Engine/Plugins/Experimental/DaySequence/Source/DaySequence/Public/DaySequenceActor.h
@@ -8,7 +8,7 @@
 #include "IMovieScenePlaybackClient.h"
 #include "MovieSceneBindingOwnerInterface.h"
 
-#include "Misc/Timecode.h"
+#include "DaySequenceTime.h"
 #include "TimerManager.h"
 
 #include "DaySequenceActor.generated.h"
@@ -480,7 +480,7 @@ protected:
 	 * Can be used as the initial time of day in PIE if bOverrideInitialTimeOfDay is true.
 	 */
 	UPROPERTY(EditAnywhere, Category = Preview, NonTransactional)
-	FTimecode TimeOfDayPreview;
+	FDaySequenceTime TimeOfDayPreview;
 #endif
 
 	/** Whether or not to run a day cycle. If this is unchecked the day cycle will remain fixed at the time specified by the Initial Time setting */
@@ -492,15 +492,15 @@ protected:
 
 	/** How long a single day cycle is */
 	UPROPERTY(EditAnywhere, Category = RuntimeDayCycle)
-	FTimecode DayLength;
+	FDaySequenceTime DayLength;
 
 	/** How long does it take for a day cycle to complete in world time. If this is the same value as day duration that means real world time is used */
 	UPROPERTY(EditAnywhere, Category = RuntimeDayCycle)
-	FTimecode TimePerCycle;
+	FDaySequenceTime TimePerCycle;
 
 	/** The initial time that the day cycle will start at */
 	UPROPERTY(EditAnywhere, Category = RuntimeDayCycle)
-	FTimecode InitialTimeOfDay;
+	FDaySequenceTime InitialTimeOfDay;
 	
 	/** Cached static time controller applied from SetStaticTimeOfDay */
 	TWeakPtr<UE::DaySequence::FStaticTimeControllerOverride> WeakTimeControllerOverride;
diff --git a/Engine/Plugins/Experimental/DaySequence/Source/DaySequence/Public/DaySequenceTime.h b/Engine/Plugins/Experimental/DaySequence/Source/DaySequence/Public/DaySequenceTime.h
new file mode 100644
--- /dev/null
+++ b/Engine/Plugins/Experimental/DaySequence/Source/DaySequence/Public/DaySequenceTime.h
@@ -0,0 +1,49 @@
+// Copyright Epic Games, Inc. All Rights Reserved.
+
+#pragma once
+
+#include "DaySequenceTime.generated.h"
+
+USTRUCT()
+struct FDaySequenceTime
+{
+	GENERATED_BODY()
+
+public:
+	static constexpr float SecondsPerHour = 3600.f;
+	static constexpr float SecondsPerMinute = 60.f;
+	
+public:
+	FDaySequenceTime(int32 InHours = 0, int32 InMinutes = 0, int32 InSeconds = 0);
+
+	bool operator==(const FDaySequenceTime& RHS) const;
+	bool operator!=(const FDaySequenceTime& RHS) const;
+
+	float ToHours() const;
+	float ToSeconds() const;
+	
+	static FDaySequenceTime FromHours(float InHours);
+	static FDaySequenceTime FromSeconds(float InSeconds);
+	
+	FString ToString(bool bForceSignDisplay = false) const
+	{
+		const bool bHasNegativeComponent = Hours < 0 || Minutes < 0 || Seconds < 0;
+		const TCHAR* SignText = bHasNegativeComponent ? TEXT("- ") : bForceSignDisplay ? TEXT("+ ") : TEXT("");
+
+		// Allocating string of length 34 to account for 2 sign chars, 10 chars for each int, and 2 ':' chars
+		TStringBuilder<34> Builder;
+		Builder.Appendf(TEXT("%s%02d:%02d:%02d"), SignText, FMath::Abs(Hours), FMath::Abs(Minutes), FMath::Abs(Seconds));
+	
+		return Builder.ToString();
+	}
+
+public:
+	UPROPERTY()
+	int32 Hours;
+
+	UPROPERTY()
+	int32 Minutes;
+
+	UPROPERTY()
+	int32 Seconds;
+};
\ No newline at end of file
diff --git a/Engine/Plugins/Experimental/DaySequence/Source/DaySequenceEditor/Private/DaySequenceEditorModule.cpp b/Engine/Plugins/Experimental/DaySequence/Source/DaySequenceEditor/Private/DaySequenceEditorModule.cpp
--- a/Engine/Plugins/Experimental/DaySequence/Source/DaySequenceEditor/Private/DaySequenceEditorModule.cpp
+++ b/Engine/Plugins/Experimental/DaySequence/Source/DaySequenceEditor/Private/DaySequenceEditorModule.cpp
@@ -11,6 +11,7 @@
 #include "DaySequenceEditorActorSpawner.h"
 #include "DaySequenceSubsystem.h"
 #include "DaySequenceConditionSetCustomization.h"
+#include "DaySequenceTimeDetailsCustomization.h"
 #include "DaySequenceEditorToolkit.h"
 #include "DaySequenceTrackEditor.h"
 #include "DaySequenceModifierComponent.h"
@@ -50,6 +51,7 @@ static const FName LevelEditorModuleName("LevelEditor");
 static const FName PropertyEditorModuleName("PropertyEditor");
 static const FName DaySequenceActorClassName("DaySequenceActor");
 static const FName DaySequenceConditionSetName("DaySequenceConditionSet");
+static const FName DaySequenceTimeName("DaySequenceTime");
 static const FName DaySequenceViewportToolBarExtensionName("DaySequenceEditorViewportToolBar");
 static const FName EnvironmentLightingActorClassName("EnvironmentLightingActor");
 
@@ -247,7 +249,8 @@ void FDaySequenceEditorModule::StartupModule()
 	PropertyModule.RegisterCustomClassLayout(DaySequenceActorClassName, FOnGetDetailCustomizationInstance::CreateStatic(&FDaySequenceActorDetails::MakeInstance));
 	PropertyModule.RegisterCustomClassLayout(EnvironmentLightingActorClassName, FOnGetDetailCustomizationInstance::CreateStatic(&FEnvironmentLightingActorDetails::MakeInstance));
 	PropertyModule.RegisterCustomPropertyTypeLayout(DaySequenceConditionSetName, FOnGetPropertyTypeCustomizationInstance::CreateStatic(&FDaySequenceConditionSetCustomization::MakeInstance));
-
+	PropertyModule.RegisterCustomPropertyTypeLayout(DaySequenceTimeName, FOnGetPropertyTypeCustomizationInstance::CreateStatic(&FDaySequenceTimeDetailsCustomization::MakeInstance));
+	
 	RegisterModulePropertySections();
 
 	RegisterEditorObjectBindings();
diff --git a/Engine/Plugins/Experimental/DaySequence/Source/DaySequenceEditor/Private/DaySequenceTimeDetailsCustomization.cpp b/Engine/Plugins/Experimental/DaySequence/Source/DaySequenceEditor/Private/DaySequenceTimeDetailsCustomization.cpp
new file mode 100644
--- /dev/null
+++ b/Engine/Plugins/Experimental/DaySequence/Source/DaySequenceEditor/Private/DaySequenceTimeDetailsCustomization.cpp
@@ -0,0 +1,102 @@
+// Copyright Epic Games, Inc. All Rights Reserved.
+
+#include "DaySequenceTimeDetailsCustomization.h"
+
+#include "Containers/Array.h"
+#include "Containers/UnrealString.h"
+#include "CoreGlobals.h"
+#include "DetailLayoutBuilder.h"
+#include "DetailWidgetRow.h"
+#include "Editor.h"
+#include "Editor/EditorEngine.h"
+#include "Fonts/SlateFontInfo.h"
+#include "HAL/Platform.h"
+#include "IDetailChildrenBuilder.h"
+#include "Internationalization/Internationalization.h"
+#include "Logging/LogCategory.h"
+#include "Logging/LogMacros.h"
+#include "Misc/Attribute.h"
+#include "Misc/CString.h"
+#include "DaySequenceTime.h"
+#include "PropertyHandle.h"
+#include "Trace/Detail/Channel.h"
+#include "UObject/UnrealType.h"
+#include "Widgets/DeclarativeSyntaxSupport.h"
+#include "Widgets/Input/SEditableTextBox.h"
+#include "Widgets/Text/STextBlock.h"
+
+#define LOCTEXT_NAMESPACE "Time"
+
+void FDaySequenceTimeDetailsCustomization::CustomizeChildren(TSharedRef<IPropertyHandle> PropertyHandle, IDetailChildrenBuilder& ChildBuilder, IPropertyTypeCustomizationUtils& CustomizationUtils)
+{
+	TimeProperty = PropertyHandle;
+
+	ChildBuilder.AddProperty(PropertyHandle)
+		.CustomWidget()
+		.NameContent()
+		[
+			PropertyHandle->CreatePropertyNameWidget()
+		]
+		.ValueContent()
+		[
+			SNew(SEditableTextBox)
+			.Text(this, &FDaySequenceTimeDetailsCustomization::OnGetTimeText)
+			.OnTextCommitted(this, &FDaySequenceTimeDetailsCustomization::OnTimeTextCommitted)
+			.SelectAllTextWhenFocused(true)
+			.RevertTextOnEscape(true)
+			.Font(IDetailLayoutBuilder::GetDetailFont())
+			.IsReadOnly_Lambda([this]()
+			{
+				return TimeProperty.IsValid() ? !TimeProperty->IsEditable() : false;
+			})
+		];
+}
+
+FText FDaySequenceTimeDetailsCustomization::OnGetTimeText() const
+{
+	TArray<void*> RawData;
+	TimeProperty->AccessRawData(RawData);
+
+	if (RawData.Num())
+	{
+		FString CurrentValue = ((FDaySequenceTime*)RawData[0])->ToString();
+		return FText::FromString(CurrentValue);
+	}
+
+	return FText::GetEmpty();
+}
+
+void FDaySequenceTimeDetailsCustomization::OnTimeTextCommitted(const FText& InText, ETextCommit::Type CommitInfo)
+{
+	TArray<void*> RawData;
+	TimeProperty->AccessRawData(RawData);
+
+	if (RawData.Num())
+	{
+		TArray<FString> Splits;
+		InText.ToString().ParseIntoArray(Splits, TEXT(":"));
+
+		const int NumSplits = Splits.Num();
+		if (NumSplits > 0 && NumSplits <= 3)
+		{
+			GEditor->BeginTransaction(FText::Format(LOCTEXT("SetTimeProperty", "Edit {0}"), TimeProperty->GetPropertyDisplayName()));
+			
+			TimeProperty->NotifyPreChange();
+			
+			((FDaySequenceTime*)RawData[0])->Hours   = FCString::Atoi(*Splits[0]);
+			((FDaySequenceTime*)RawData[0])->Minutes = NumSplits > 1 ? FCString::Atoi(*Splits[1]) : 0;
+			((FDaySequenceTime*)RawData[0])->Seconds = NumSplits > 2 ? FCString::Atoi(*Splits[2]) : 0;
+			
+			TimeProperty->NotifyPostChange(EPropertyChangeType::ValueSet);
+			TimeProperty->NotifyFinishedChangingProperties();
+			
+			GEditor->EndTransaction();
+		}
+		else
+		{
+			UE_LOG(LogTemp, Error, TEXT("Unexpected time format. Expected 1-3 values, got %d"), Splits.Num());
+		}
+	}
+}
+
+#undef LOCTEXT_NAMESPACE
diff --git a/Engine/Plugins/Experimental/DaySequence/Source/DaySequenceEditor/Public/DaySequenceTimeDetailsCustomization.h b/Engine/Plugins/Experimental/DaySequence/Source/DaySequenceEditor/Public/DaySequenceTimeDetailsCustomization.h
new file mode 100644
--- /dev/null
+++ b/Engine/Plugins/Experimental/DaySequence/Source/DaySequenceEditor/Public/DaySequenceTimeDetailsCustomization.h
@@ -0,0 +1,38 @@
+// Copyright Epic Games, Inc. All Rights Reserved.
+
+#pragma once
+
+#include "IPropertyTypeCustomization.h"
+#include "Internationalization/Text.h"
+#include "Templates/SharedPointer.h"
+#include "Types/SlateEnums.h"
+
+class FDetailWidgetRow;
+class IDetailChildrenBuilder;
+class IPropertyHandle;
+
+/**
+ *  Customize FDaySequenceTime
+ */
+class FDaySequenceTimeDetailsCustomization : public IPropertyTypeCustomization
+{
+public:
+	static TSharedRef<IPropertyTypeCustomization> MakeInstance()
+	{
+		return MakeShared<FDaySequenceTimeDetailsCustomization>();
+	}
+
+	FDaySequenceTimeDetailsCustomization()
+	{}
+
+	/** IPropertyTypeCustomization interface */
+	virtual void CustomizeHeader(TSharedRef<IPropertyHandle> PropertyHandle, FDetailWidgetRow& HeaderRow, IPropertyTypeCustomizationUtils& CustomizationUtils) override {}
+	virtual void CustomizeChildren(TSharedRef<IPropertyHandle> PropertyHandle, IDetailChildrenBuilder& ChildBuilder, IPropertyTypeCustomizationUtils& CustomizationUtils) override;
+
+private:
+	FText OnGetTimeText() const;
+	void OnTimeTextCommitted(const FText& InText, ETextCommit::Type CommitInfo);
+
+	/** Store the property handle to the Time field so we can get/set the value on the object via text box callbacks. */
+	TSharedPtr<IPropertyHandle> TimeProperty;
+};
diff --git a/Engine/Source/Runtime/Core/Public/Misc/Timecode.h b/Engine/Source/Runtime/Core/Public/Misc/Timecode.h
--- a/Engine/Source/Runtime/Core/Public/Misc/Timecode.h
+++ b/Engine/Source/Runtime/Core/Public/Misc/Timecode.h
@@ -392,7 +392,7 @@ public:
 		{
 		    return FMath::Abs(Val) >= MinVal && FMath::Abs(Val) <= MaxVal;
 		};
-		return InAbsRange(Hours,0,24) && InAbsRange(Minutes, 0, 59) && InAbsRange(Seconds, 0, 59) && Subframe >= 0;
+		return InAbsRange(Hours,0,23) && InAbsRange(Minutes, 0, 59) && InAbsRange(Seconds, 0, 59) && Subframe >= 0;
 	}
 public:
 
