DaySequenceCollectionAsset: Properly initialize BiasOffset for default constructed collection entries

diff --git a/Engine/Plugins/Experimental/DaySequence/Source/DaySequence/Public/DaySequenceCollectionAsset.h b/Engine/Plugins/Experimental/DaySequence/Source/DaySequence/Public/DaySequenceCollectionAsset.h
--- a/Engine/Plugins/Experimental/DaySequence/Source/DaySequence/Public/DaySequenceCollectionAsset.h
+++ b/Engine/Plugins/Experimental/DaySequence/Source/DaySequence/Public/DaySequenceCollectionAsset.h
@@ -13,10 +13,8 @@ USTRUCT()
 struct FDaySequenceCollectionEntry
 {
 	GENERATED_BODY()
-
-	FDaySequenceCollectionEntry() = default;
-
-	FDaySequenceCollectionEntry(UDaySequence* InDaySequence)
+	
+	FDaySequenceCollectionEntry(UDaySequence* InDaySequence = nullptr)
 	: Sequence(InDaySequence)
 	, BiasOffset(0)
 	, Conditions(FDaySequenceConditionSet())
