DaySequenceConditionTag: Revise API

Remove GetConditionName (in favor of ConditionName, an EditDefaultsOnly field) and rename SetupBroadcastBindings to SetupOnConditionValueChanged.

#changelist validated
#virtualized

diff --git a/Engine/Plugins/Experimental/DaySequence/Source/DaySequence/Private/DaySequenceConditionTag.cpp b/Engine/Plugins/Experimental/DaySequence/Source/DaySequence/Private/DaySequenceConditionTag.cpp
--- a/Engine/Plugins/Experimental/DaySequence/Source/DaySequence/Private/DaySequenceConditionTag.cpp
+++ b/Engine/Plugins/Experimental/DaySequence/Source/DaySequence/Private/DaySequenceConditionTag.cpp
@@ -4,7 +4,7 @@
 
 void UDaySequenceConditionTag::Initialize()
 {
-	SetupBroadcastBindings();
+	SetupOnConditionValueChanged();
 }
 
 bool UDaySequenceConditionTag::Evaluate_Implementation() const
@@ -12,9 +12,9 @@ bool UDaySequenceConditionTag::Evaluate_Implementation() const
 	return true;
 }
 
-FString UDaySequenceConditionTag::GetConditionName_Implementation() const
+FString UDaySequenceConditionTag::GetConditionName() const
 {
-	return GetClass()->GetName();
+	return ConditionName.IsEmpty() ? GetClass()->GetName() : ConditionName;
 }
 
 UWorld* UDaySequenceConditionTag::GetWorld() const
@@ -27,7 +27,7 @@ UWorld* UDaySequenceConditionTag::GetWorld() const
 	return GetOuter()->GetWorld();
 }
 
-void UDaySequenceConditionTag::SetupBroadcastBindings_Implementation() const
+void UDaySequenceConditionTag::SetupOnConditionValueChanged_Implementation() const
 {
 }
 
diff --git a/Engine/Plugins/Experimental/DaySequence/Source/DaySequence/Public/DaySequenceConditionTag.h b/Engine/Plugins/Experimental/DaySequence/Source/DaySequence/Public/DaySequenceConditionTag.h
--- a/Engine/Plugins/Experimental/DaySequence/Source/DaySequence/Public/DaySequenceConditionTag.h
+++ b/Engine/Plugins/Experimental/DaySequence/Source/DaySequence/Public/DaySequenceConditionTag.h
@@ -26,13 +26,8 @@ public:
 	UFUNCTION(BlueprintNativeEvent, Category = "General")
 	bool Evaluate() const;
 
-	/**
-	 * Derived classes should override this function to specify a friendly name which is used
-	 * when prompting users with a list of possible conditions to apply to a given sequence.
-	 */
-	UFUNCTION(BlueprintNativeEvent, Category = "General")
 	FString GetConditionName() const;
-
+	
 	DECLARE_MULTICAST_DELEGATE(FOnConditionValueChanged);
 	FOnConditionValueChanged& GetOnConditionValueChanged() { return OnConditionValueChanged; }
 	
@@ -46,7 +41,7 @@ protected:
 	 * can propagate those broadcasts to notify users of this condition that the condition needs reevaluating.
 	 */
 	UFUNCTION(BlueprintNativeEvent, Category = "General")
-	void SetupBroadcastBindings() const;
+	void SetupOnConditionValueChanged() const;
 	
 	/**
 	 * Derived classes should call this function to notify listeners that the underlying condition may have changed.
@@ -55,6 +50,14 @@ protected:
 	UFUNCTION(BlueprintCallable, Category = "General")
 	void BroadcastOnConditionValueChanged();
 
+protected:
+	/**
+	 * Derived classes should give this a meaningful default value which is displayed
+	 * when prompting users with a list of possible conditions to apply to a given sequence.
+	 */
+	UPROPERTY(EditDefaultsOnly, Category = "General")
+	FString ConditionName;
+
 private:
 	FOnConditionValueChanged OnConditionValueChanged;
 
