DaySequence: Add DaySequence CSV stat category

diff --git a/Engine/Plugins/Experimental/DaySequence/Source/DaySequence/Private/DaySequenceActor.cpp b/Engine/Plugins/Experimental/DaySequence/Source/DaySequence/Private/DaySequenceActor.cpp
--- a/Engine/Plugins/Experimental/DaySequence/Source/DaySequence/Private/DaySequenceActor.cpp
+++ b/Engine/Plugins/Experimental/DaySequence/Source/DaySequence/Private/DaySequenceActor.cpp
@@ -538,7 +538,7 @@ void ADaySequenceActor::InitializePlayer()
 
 void ADaySequenceActor::InitializeRootSequence()
 {
-	CSV_SCOPED_TIMING_STAT_EXCLUSIVE(DaySequenceActor_InitializeRootSequence);
+	CSV_SCOPED_TIMING_STAT(DaySequence, InitializeRootSequence);
 
 	if (IsTemplate())
 	{
@@ -854,7 +854,7 @@ void ADaySequenceActor::UpdateSubSectionTimeScale(UMovieSceneSubSection* InSubSe
 
 void ADaySequenceActor::OnSequencePlayerUpdate(const UMovieSceneSequencePlayer& Player, FFrameTime CurrentTime, FFrameTime PreviousTime)
 {
-	CSV_SCOPED_TIMING_STAT_EXCLUSIVE(DaySequenceActor_OnSequencePlayerUpdate);
+	CSV_SCOPED_TIMING_STAT(DaySequence, OnSequencePlayerUpdate);
 	
 	auto FrameTimeToDayHours = [this](const FFrameTime& Time)
 	{
diff --git a/Engine/Plugins/Experimental/DaySequence/Source/DaySequence/Private/DaySequenceModifierComponent.cpp b/Engine/Plugins/Experimental/DaySequence/Source/DaySequence/Private/DaySequenceModifierComponent.cpp
--- a/Engine/Plugins/Experimental/DaySequence/Source/DaySequence/Private/DaySequenceModifierComponent.cpp
+++ b/Engine/Plugins/Experimental/DaySequence/Source/DaySequence/Private/DaySequenceModifierComponent.cpp
@@ -5,6 +5,7 @@
 #include "Components/BoxComponent.h"
 #include "DaySequence.h"
 #include "DaySequenceCollectionAsset.h"
+#include "DaySequenceModule.h"
 #include "DaySequenceTrack.h"
 
 #include "Engine/World.h"
@@ -406,7 +407,7 @@ void UDaySequenceModifierComponent::OnUnregister()
 
 void UDaySequenceModifierComponent::DaySequenceUpdate()
 {
-	CSV_SCOPED_TIMING_STAT_EXCLUSIVE(DaySequenceModifier_SequencePlayerUpdated);
+	CSV_SCOPED_TIMING_STAT(DaySequence, SequencePlayerUpdated);
 	
 	// Force expensive update
 	const float DistanceBlendFactor = UpdateBlendWeight();
@@ -589,8 +590,6 @@ void UDaySequenceModifierComponent::DisableComponent()
 
 void UDaySequenceModifierComponent::EnableModifier()
 {
-	CSV_SCOPED_TIMING_STAT_EXCLUSIVE(DaySequenceModifierComponent_EnableModifier);
-	
 	if (bIsEnabled || !CanBeEnabled())
 	{
 		return;
@@ -625,8 +624,6 @@ void UDaySequenceModifierComponent::EnableModifier()
 
 void UDaySequenceModifierComponent::DisableModifier()
 {
-	CSV_SCOPED_TIMING_STAT_EXCLUSIVE(DaySequenceModifierComponent_DisableModifier);
-	
 	if (!bIsEnabled)
 	{
 		return;
@@ -715,6 +712,8 @@ void UDaySequenceModifierComponent::SetInitialTimeOfDay()
 
 void UDaySequenceModifierComponent::ReinitializeSubSequence(ADaySequenceActor::FSubSectionPreserveMap* SectionsToPreserve)
 {
+	CSV_SCOPED_TIMING_STAT(DaySequence, ReinitializeSubSequence);
+	
 #if ROOT_SEQUENCE_RECONSTRUCTION_ENABLED
 	bool bReinit = true;
 
@@ -1336,7 +1335,7 @@ void UDaySequenceModifierComponent::SetUserDaySequence(UDaySequence* InDaySequen
 
 bool UDaySequenceModifierComponent::GetBlendPosition(FVector& InPosition) const
 {
-	CSV_SCOPED_TIMING_STAT_EXCLUSIVE(DaySequenceModifier_GetBlendPosition);
+	CSV_SCOPED_TIMING_STAT(DaySequence, GetBlendPosition);
 	
 #if WITH_EDITOR
 	if (const UWorld* World = GetWorld(); World && !World->IsGameWorld())
@@ -1348,7 +1347,7 @@ bool UDaySequenceModifierComponent::GetBlendPosition(FVector& InPosition) const
 #endif
 	if (const APlayerController* BlendTarget = ExternalVolumeBlendTarget.Get())
 	{
-		CSV_SCOPED_TIMING_STAT_EXCLUSIVE(DaySequenceModifier_GetPlayerViewPointScope);
+		CSV_SCOPED_TIMING_STAT(DaySequence, GetPlayerViewPoint);
 		InPosition = BlendTarget->PlayerCameraManager->GetCameraLocation();
 		return true;
 	}
@@ -1364,7 +1363,7 @@ float UDaySequenceModifierComponent::GetDistanceBlendFactorForShape(const UShape
 
 float UDaySequenceModifierComponent::GetDistanceBlendFactor(const FVector& Position) const
 {
-	CSV_SCOPED_TIMING_STAT_EXCLUSIVE(DaySequenceModifier_GetDistanceBlendFactor);
+	CSV_SCOPED_TIMING_STAT(DaySequence, GetDistanceBlendFactor);
 	
 	CachedDistanceBlendFactor = 0.f;
 
diff --git a/Engine/Plugins/Experimental/DaySequence/Source/DaySequence/Private/DaySequenceModule.cpp b/Engine/Plugins/Experimental/DaySequence/Source/DaySequence/Private/DaySequenceModule.cpp
--- a/Engine/Plugins/Experimental/DaySequence/Source/DaySequence/Private/DaySequenceModule.cpp
+++ b/Engine/Plugins/Experimental/DaySequence/Source/DaySequence/Private/DaySequenceModule.cpp
@@ -11,6 +11,7 @@
 #define LOCTEXT_NAMESPACE "FDaySequenceModule"
 
 DEFINE_LOG_CATEGORY(LogDaySequence);
+CSV_DEFINE_CATEGORY(DaySequence, false);
 
 namespace UE::DaySequence
 {
diff --git a/Engine/Plugins/Experimental/DaySequence/Source/DaySequence/Private/DaySequenceModule.h b/Engine/Plugins/Experimental/DaySequence/Source/DaySequence/Private/DaySequenceModule.h
--- a/Engine/Plugins/Experimental/DaySequence/Source/DaySequence/Private/DaySequenceModule.h
+++ b/Engine/Plugins/Experimental/DaySequence/Source/DaySequence/Private/DaySequenceModule.h
@@ -3,8 +3,10 @@
 #pragma once
 
 #include "IDaySequenceModule.h"
+#include "ProfilingDebugging/CsvProfiler.h"
 
 DECLARE_LOG_CATEGORY_EXTERN(LogDaySequence, Log, All);
+CSV_DECLARE_CATEGORY_EXTERN(DaySequence);
 
 class FDaySequenceModule : public IDaySequenceModule
 {
