DaySequence: Hide DaySequenceEditor tab.

diff --git a/Engine/Plugins/Experimental/DaySequence/Source/DaySequenceEditor/Private/DaySequenceEditorCommands.cpp b/Engine/Plugins/Experimental/DaySequence/Source/DaySequenceEditor/Private/DaySequenceEditorCommands.cpp
--- a/Engine/Plugins/Experimental/DaySequence/Source/DaySequenceEditor/Private/DaySequenceEditorCommands.cpp
+++ b/Engine/Plugins/Experimental/DaySequence/Source/DaySequenceEditor/Private/DaySequenceEditorCommands.cpp
@@ -9,7 +9,6 @@ void FDaySequenceEditorCommands::RegisterCommands()
 	UI_COMMAND(OverrideInitialTimeOfDay, "Override Initial Time of Day", "Use the current preview time as the initial time of day in PIE", EUserInterfaceActionType::ToggleButton, FInputChord());
 	UI_COMMAND(OverrideRunDayCycle, "Freeze Time in PIE", "Pauses the day cycle when starting PIE", EUserInterfaceActionType::ToggleButton, FInputChord());
 	
-	UI_COMMAND(OpenDaySequenceEditor, "DaySequence", "Opens the Day Sequence Editor", EUserInterfaceActionType::Button, FInputChord());
 	UI_COMMAND(SelectDaySequenceActor, "Select Day Sequence Actor", "Selects the Day Sequence Actor", EUserInterfaceActionType::Button, FInputChord());
 	UI_COMMAND(RefreshDaySequenceActor, "Refresh Day Sequence Actor", "Refreshes the Main Sequence in the Day Sequence Actor", EUserInterfaceActionType::Button, FInputChord());
 	UI_COMMAND(OpenDaySequenceActor, "Open Day Sequence Actor", "Opens the Day Sequence Actor blueprint", EUserInterfaceActionType::Button, FInputChord());
diff --git a/Engine/Plugins/Experimental/DaySequence/Source/DaySequenceEditor/Private/DaySequenceEditorCommands.h b/Engine/Plugins/Experimental/DaySequence/Source/DaySequenceEditor/Private/DaySequenceEditorCommands.h
--- a/Engine/Plugins/Experimental/DaySequence/Source/DaySequenceEditor/Private/DaySequenceEditorCommands.h
+++ b/Engine/Plugins/Experimental/DaySequence/Source/DaySequenceEditor/Private/DaySequenceEditorCommands.h
@@ -21,7 +21,6 @@ public:
 	TSharedPtr<FUICommandInfo> OverrideInitialTimeOfDay;
 	TSharedPtr<FUICommandInfo> OverrideRunDayCycle;
 	
-	TSharedPtr<FUICommandInfo> OpenDaySequenceEditor;
 	TSharedPtr<FUICommandInfo> SelectDaySequenceActor;
 	TSharedPtr<FUICommandInfo> RefreshDaySequenceActor;
 	TSharedPtr<FUICommandInfo> OpenDaySequenceActor;
diff --git a/Engine/Plugins/Experimental/DaySequence/Source/DaySequenceEditor/Private/DaySequenceEditorModule.cpp b/Engine/Plugins/Experimental/DaySequence/Source/DaySequenceEditor/Private/DaySequenceEditorModule.cpp
--- a/Engine/Plugins/Experimental/DaySequence/Source/DaySequenceEditor/Private/DaySequenceEditorModule.cpp
+++ b/Engine/Plugins/Experimental/DaySequence/Source/DaySequenceEditor/Private/DaySequenceEditorModule.cpp
@@ -17,7 +17,6 @@
 #include "DaySequenceModifierComponent.h"
 #include "Editor/UnrealEdTypes.h"
 #include "MovieSceneSequenceEditor_DaySequence.h"
-#include "SDaySequenceSettings.h"
 #include "IDaySequenceModule.h"
 #include "EnvironmentLightingActorDetails.h"
 
@@ -27,7 +26,6 @@
 
 #include "Application/ThrottleManager.h"
 #include "Framework/Docking/TabManager.h"
-#include "LevelEditor.h"
 #include "ToolMenu.h"
 #include "ViewportToolBarContext.h"
 #include "PropertyEditorModule.h"
@@ -46,8 +44,6 @@
 
 #define LOCTEXT_NAMESPACE "DaySequenceEditor"
 
-static const FName DaySequenceTabName("DaySequenceEditor");
-static const FName LevelEditorModuleName("LevelEditor");
 static const FName PropertyEditorModuleName("PropertyEditor");
 static const FName DaySequenceActorClassName("DaySequenceActor");
 static const FName DaySequenceConditionSetName("DaySequenceConditionSet");
@@ -62,6 +58,8 @@ public:
 	virtual void StartupModule() override;
 	virtual void ShutdownModule() override;
 
+	void OnPostEngineInit();
+	
 	// FGCObject interface
 	virtual void AddReferencedObjects( FReferenceCollector& Collector ) override;
 	virtual FString GetReferencerName() const override;
@@ -72,7 +70,6 @@ public:
 	virtual FPreSelectDaySequenceActor& OnPreSelectDaySequenceActor() override;
 	virtual FPostSelectDaySequenceActor& OnPostSelectDaySequenceActor() override;
 	
-	void OnOpenDaySequenceEditor();
 	bool IsDaySequenceActorPreviewEnabled() const;
 	void OnOpenRootSequence();
 	bool CanOpenRootSequence() const;
@@ -108,7 +105,6 @@ private:
 	TSharedRef<SWidget> CreateDaySequenceViewportToolbarExtension(FWeakObjectPtr ExtensionContext);
 	
 	TSharedRef<SWidget> CreateDaySequencePreviewWidget();
-	TSharedRef<SDockTab> CreateDaySequenceTab(const FSpawnTabArgs& Args);
 
 	void OnEditorCameraMoved(const FVector& Location, const FRotator& Rotation, ELevelViewportType Type, int32 ViewIndex);
 	
@@ -141,6 +137,8 @@ private:
 void FDaySequenceEditorModule::StartupModule()
 {
 	// This code will execute after your module is loaded into memory; the exact timing is specified in the .uplugin file per-module
+
+	FCoreDelegates::OnPostEngineInit.AddRaw(this, &FDaySequenceEditorModule::OnPostEngineInit);
 	
 	FDaySequenceEditorStyle::Initialize();
 	FDaySequenceEditorCommands::Register();
@@ -149,11 +147,6 @@ void FDaySequenceEditorModule::StartupModule()
 
 	PluginCommands = MakeShareable(new FUICommandList);
 
-	PluginCommands->MapAction(
-		FDaySequenceEditorCommands::Get().OpenDaySequenceEditor,
-		FExecuteAction::CreateRaw(this, &FDaySequenceEditorModule::OnOpenDaySequenceEditor),
-		FCanExecuteAction());
-
 	PluginCommands->MapAction(
 		FDaySequenceEditorCommands::Get().OverrideInitialTimeOfDay,
 		FExecuteAction::CreateLambda([this]()
@@ -213,35 +206,6 @@ void FDaySequenceEditorModule::StartupModule()
 		FDaySequenceEditorCommands::Get().OpenDaySequenceActor,
 		FExecuteAction::CreateRaw(this, &FDaySequenceEditorModule::OnOpenDaySequenceActor),
 		FCanExecuteAction::CreateRaw(this, &FDaySequenceEditorModule::CanOpenDaySequenceActor));
-	{
-		const IWorkspaceMenuStructure& MenuStructure = WorkspaceMenu::GetMenuStructure();
-		TSharedRef<FWorkspaceItem> LevelEditorGroup = MenuStructure.GetLevelEditorCategory();
-
-		FLevelEditorModule& LevelEditorModule = FModuleManager::LoadModuleChecked<FLevelEditorModule>(LevelEditorModuleName);
-		
-		auto RegisterTabSpawner = [this, LevelEditorGroup]()
-		{
-			FLevelEditorModule& LevelEditorModule = FModuleManager::LoadModuleChecked<FLevelEditorModule>(LevelEditorModuleName);
-			TSharedPtr<FTabManager> LevelEditorTabManager = LevelEditorModule.GetLevelEditorTabManager();
-
-			LevelEditorTabManager->RegisterTabSpawner(DaySequenceTabName, FOnSpawnTab::CreateRaw(this, &FDaySequenceEditorModule::CreateDaySequenceTab))
-				.SetDisplayName(LOCTEXT("TabTitle", "Day Sequence"))
-				.SetTooltipText(LOCTEXT("TooltipText", "Opens the Day Sequence editor"))
-				.SetGroup(LevelEditorGroup)
-				.SetIcon(FSlateIcon(FDaySequenceEditorStyle::GetStyleSetName(), "DaySequenceEditor.OpenDaySequenceEditor"));
-
-			RegisterMenus();
-		};
-
-		if (TSharedPtr<FTabManager> LevelEditorTabManager = LevelEditorModule.GetLevelEditorTabManager())
-		{
-			RegisterTabSpawner();
-		}
-		else
-		{
-			LevelEditorModule.OnTabManagerChanged().AddLambda(RegisterTabSpawner);
-		}
-	}
 
 	UToolMenus::RegisterStartupCallback(FSimpleMulticastDelegate::FDelegate::CreateRaw(this, &FDaySequenceEditorModule::RegisterMenus));
 
@@ -266,21 +230,7 @@ void FDaySequenceEditorModule::StartupModule()
 
 void FDaySequenceEditorModule::ShutdownModule()
 {
-	if (FSlateApplication::IsInitialized() && FModuleManager::Get().IsModuleLoaded(LevelEditorModuleName))
-	{
-		FLevelEditorModule* LevelEditorModule = FModuleManager::GetModulePtr<FLevelEditorModule>(LevelEditorModuleName);
-		TSharedPtr<FTabManager> LevelEditorTabManager;
-		if (LevelEditorModule)
-		{
-			LevelEditorTabManager = LevelEditorModule->GetLevelEditorTabManager();
-			LevelEditorModule->OnTabManagerChanged().RemoveAll(this);
-		}
-
-		if (LevelEditorTabManager.IsValid())
-		{
-			LevelEditorTabManager->UnregisterTabSpawner(DaySequenceTabName);
-		}
-	}
+	FCoreDelegates::OnPostEngineInit.RemoveAll(this);
 
 	FEditorDelegates::OnEditorCameraMoved.Remove(OnEditorCameraMovedHandle);
 
@@ -314,6 +264,23 @@ void FDaySequenceEditorModule::ShutdownModule()
 	FDaySequenceEditorCommands::Unregister();
 }
 
+void FDaySequenceEditorModule::OnPostEngineInit()
+{
+	if (GEditor)
+	{
+		if (UPanelExtensionSubsystem* PanelExtensionSubsystem = GEditor->GetEditorSubsystem<UPanelExtensionSubsystem>())
+		{
+			if (!PanelExtensionSubsystem->IsPanelFactoryRegistered(DaySequenceViewportToolBarExtensionName))
+			{
+				FPanelExtensionFactory DaySequenceViewportMenuWidget;
+				DaySequenceViewportMenuWidget.CreateExtensionWidget = FPanelExtensionFactory::FCreateExtensionWidget::CreateRaw(this, &FDaySequenceEditorModule::CreateDaySequenceViewportToolbarExtension);
+				DaySequenceViewportMenuWidget.Identifier = DaySequenceViewportToolBarExtensionName;
+				PanelExtensionSubsystem->RegisterPanelFactory("LevelViewportToolBar.LeftExtension", DaySequenceViewportMenuWidget);
+			}
+		}
+	}
+}
+
 void FDaySequenceEditorModule::AddReferencedObjects( FReferenceCollector& Collector )
 {
 	if (Settings)
@@ -347,16 +314,6 @@ IDaySequenceEditorModule::FPostSelectDaySequenceActor& FDaySequenceEditorModule:
 	return OnPostSelectDaySequenceActorDelegate;
 }
 
-void FDaySequenceEditorModule::OnOpenDaySequenceEditor()
-{
-	FLevelEditorModule& LevelEditorModule = FModuleManager::LoadModuleChecked<FLevelEditorModule>(LevelEditorModuleName);
-
-	if (TSharedPtr<FTabManager> LevelEditorTabManager = LevelEditorModule.GetLevelEditorTabManager())
-	{
-		LevelEditorTabManager->TryInvokeTab(DaySequenceTabName);
-	}
-}
-
 bool FDaySequenceEditorModule::IsDaySequenceActorPreviewEnabled() const
 {
 	return DaySequenceActorPreview.IsPreviewEnabled();
@@ -507,20 +464,6 @@ TSharedPtr<FPropertySection> FDaySequenceEditorModule::RegisterPropertySection(F
 
 void FDaySequenceEditorModule::RegisterMenus()
 {
-	if (GEditor)
-	{
-		if (UPanelExtensionSubsystem* PanelExtensionSubsystem = GEditor->GetEditorSubsystem<UPanelExtensionSubsystem>())
-		{
-			if (!PanelExtensionSubsystem->IsPanelFactoryRegistered(DaySequenceViewportToolBarExtensionName))
-			{
-				FPanelExtensionFactory DaySequenceViewportMenuWidget;
-				DaySequenceViewportMenuWidget.CreateExtensionWidget = FPanelExtensionFactory::FCreateExtensionWidget::CreateRaw(this, &FDaySequenceEditorModule::CreateDaySequenceViewportToolbarExtension);
-				DaySequenceViewportMenuWidget.Identifier = DaySequenceViewportToolBarExtensionName;
-				PanelExtensionSubsystem->RegisterPanelFactory("LevelViewportToolBar.LeftExtension", DaySequenceViewportMenuWidget);
-			}
-		}
-	}
-	
 	// Owner will be used for cleanup in call to UToolMenus::UnregisterOwner
 	FToolMenuOwnerScoped OwnerScoped(this);
 	{
@@ -654,15 +597,6 @@ TSharedRef<SWidget> FDaySequenceEditorModule::CreateDaySequencePreviewWidget()
 		];
 }
 
-TSharedRef<SDockTab> FDaySequenceEditorModule::CreateDaySequenceTab(const FSpawnTabArgs& Args)
-{
-	return
-		SNew(SDockTab)
-		[
-			SNew(SDaySequenceSettings)
-		];
-}
-
 void FDaySequenceEditorModule::RegisterModulePropertySections()
 {
 	FPropertyEditorModule& PropertyModule = FModuleManager::GetModuleChecked<FPropertyEditorModule>(PropertyEditorModuleName);
