DaySequenceModifierComponent: Add support for Simulate in Editor

diff --git a/Engine/Plugins/Experimental/DaySequence/Source/DaySequence/Private/DaySequenceModifierComponent.cpp b/Engine/Plugins/Experimental/DaySequence/Source/DaySequence/Private/DaySequenceModifierComponent.cpp
--- a/Engine/Plugins/Experimental/DaySequence/Source/DaySequence/Private/DaySequenceModifierComponent.cpp
+++ b/Engine/Plugins/Experimental/DaySequence/Source/DaySequence/Private/DaySequenceModifierComponent.cpp
@@ -129,6 +129,7 @@ namespace UE::DaySequence
 	}
 
 	FVector GVolumePreviewLocation = FVector::ZeroVector;
+	bool bIsSimulating = false;
 
 	float ComputeBoxSignedDistance(const UBoxComponent* BoxComponent, const FVector& InWorldPosition)
 	{
@@ -290,6 +291,11 @@ void UDaySequenceModifierComponent::SetVolumePreviewLocation(const FVector& Loca
 	UE::DaySequence::GVolumePreviewLocation = Location;
 }
 
+void UDaySequenceModifierComponent::SetIsSimulating(bool bInIsSimulating)
+{
+	UE::DaySequence::bIsSimulating = bInIsSimulating;
+}
+
 void UDaySequenceModifierComponent::UpdateEditorPreview(float DeltaTime)
 {
 	using namespace UE::DaySequence;
@@ -1338,7 +1344,7 @@ bool UDaySequenceModifierComponent::GetBlendPosition(FVector& InPosition) const
 	CSV_SCOPED_TIMING_STAT(DaySequence, GetBlendPosition);
 	
 #if WITH_EDITOR
-	if (const UWorld* World = GetWorld(); World && !World->IsGameWorld())
+	if (const UWorld* World = GetWorld(); World && (!World->IsGameWorld() || UE::DaySequence::bIsSimulating))
 	{
 		InPosition = UE::DaySequence::GVolumePreviewLocation;
 		return true;
diff --git a/Engine/Plugins/Experimental/DaySequence/Source/DaySequence/Public/DaySequenceModifierComponent.h b/Engine/Plugins/Experimental/DaySequence/Source/DaySequence/Public/DaySequenceModifierComponent.h
--- a/Engine/Plugins/Experimental/DaySequence/Source/DaySequence/Public/DaySequenceModifierComponent.h
+++ b/Engine/Plugins/Experimental/DaySequence/Source/DaySequence/Public/DaySequenceModifierComponent.h
@@ -117,6 +117,7 @@ public:
 
 #if WITH_EDITOR
 	DAYSEQUENCE_API static void SetVolumePreviewLocation(const FVector& Location);
+	DAYSEQUENCE_API static void SetIsSimulating(bool bInIsSimulating);
 #endif
 
 	/**
diff --git a/Engine/Plugins/Experimental/DaySequence/Source/DaySequenceEditor/Private/DaySequenceEditorModule.cpp b/Engine/Plugins/Experimental/DaySequence/Source/DaySequenceEditor/Private/DaySequenceEditorModule.cpp
--- a/Engine/Plugins/Experimental/DaySequence/Source/DaySequenceEditor/Private/DaySequenceEditorModule.cpp
+++ b/Engine/Plugins/Experimental/DaySequence/Source/DaySequenceEditor/Private/DaySequenceEditorModule.cpp
@@ -131,6 +131,12 @@ private:
 
 	FDelegateHandle OnEditorCameraMovedHandle;
 
+	FDelegateHandle OnBeginPIEHandle;
+	
+	FDelegateHandle OnEndPIEHandle;
+
+	FDelegateHandle OnSwitchPIEAndSIEHandle;
+
 	int32 LastAllowThrottling = 1;
 };
 
@@ -145,6 +151,21 @@ void FDaySequenceEditorModule::StartupModule()
 
 	OnEditorCameraMovedHandle = FEditorDelegates::OnEditorCameraMoved.AddRaw(this, &FDaySequenceEditorModule::OnEditorCameraMoved);
 
+	OnBeginPIEHandle = FEditorDelegates::BeginPIE.AddLambda([](bool bIsSimulating)
+	{
+		UDaySequenceModifierComponent::SetIsSimulating(bIsSimulating);
+	});
+	
+	OnEndPIEHandle = FEditorDelegates::EndPIE.AddLambda([](bool)
+	{
+		UDaySequenceModifierComponent::SetIsSimulating(false);
+	});
+
+	OnSwitchPIEAndSIEHandle = FEditorDelegates::OnSwitchBeginPIEAndSIE.AddLambda([](bool bIsSimulating)
+	{
+		UDaySequenceModifierComponent::SetIsSimulating(bIsSimulating);
+	});
+	
 	PluginCommands = MakeShareable(new FUICommandList);
 
 	PluginCommands->MapAction(
@@ -234,6 +255,12 @@ void FDaySequenceEditorModule::ShutdownModule()
 
 	FEditorDelegates::OnEditorCameraMoved.Remove(OnEditorCameraMovedHandle);
 
+	FEditorDelegates::BeginPIE.Remove(OnBeginPIEHandle);
+	
+	FEditorDelegates::EndPIE.Remove(OnEndPIEHandle);
+
+	FEditorDelegates::OnSwitchBeginPIEAndSIE.Remove(OnSwitchPIEAndSIEHandle);
+
 	if (FPropertyEditorModule* PropertyModule = FModuleManager::GetModulePtr<FPropertyEditorModule>(PropertyEditorModuleName))
 	{
 		PropertyModule->UnregisterCustomClassLayout(DaySequenceActorClassName);
