DaySequence: Remove DaySequenceHelpers

DaySequenceHelpers becomes obsolete with the introduction of FDaySequenceTime

diff --git a/Engine/Plugins/Experimental/DaySequence/Source/DaySequence/Private/DaySequenceActor.cpp b/Engine/Plugins/Experimental/DaySequence/Source/DaySequence/Private/DaySequenceActor.cpp
--- a/Engine/Plugins/Experimental/DaySequence/Source/DaySequence/Private/DaySequenceActor.cpp
+++ b/Engine/Plugins/Experimental/DaySequence/Source/DaySequence/Private/DaySequenceActor.cpp
@@ -2,7 +2,6 @@
 
 #include "DaySequenceActor.h"
 #include "DaySequenceCollectionAsset.h"
-#include "DaySequenceHelpers.h"
 #include "DaySequence.h"
 #include "DaySequenceModule.h"
 #include "DaySequencePlayer.h"
@@ -467,8 +466,7 @@ void ADaySequenceActor::ConditionalSetTimeOfDayPreview(float InHours)
 {
 	// Wrap the input hours using day length
 	InHours = FMath::Frac(InHours / GetDayLength()) * GetDayLength(); 
-	const TOptional<FFrameRate> FrameRate = RootSequence ? RootSequence->GetMovieScene()->GetDisplayRate() : TOptional<FFrameRate>();
-	const FDaySequenceTime NewTimeOfDayPreview = DaySequenceHelpers::HoursToTimecode(InHours, FrameRate);
+	const FDaySequenceTime NewTimeOfDayPreview = FDaySequenceTime::FromHours(InHours);
 	if (NewTimeOfDayPreview != TimeOfDayPreview)
 	{
 		TimeOfDayPreview = NewTimeOfDayPreview;
@@ -482,8 +480,7 @@ void ADaySequenceActor::ConditionalSetTimeOfDayPreview(float InHours)
 float ADaySequenceActor::GetTimeOfDayPreview() const
 {
 #if WITH_EDITOR
-	const TOptional<FFrameRate> FrameRate = RootSequence ? RootSequence->GetMovieScene()->GetDisplayRate() : TOptional<FFrameRate>();
-	return DaySequenceHelpers::TimecodeToHours(TimeOfDayPreview, FrameRate);
+	return TimeOfDayPreview.ToHours();
 #else
 	return 0.0f;
 #endif
@@ -494,8 +491,7 @@ void ADaySequenceActor::SetTimeOfDayPreview(float InHours)
 #if WITH_EDITOR
 	// Wrap the input hours using day length
 	InHours = FMath::Frac(InHours / GetDayLength()) * GetDayLength(); 
-	const TOptional<FFrameRate> FrameRate = RootSequence ? RootSequence->GetMovieScene()->GetDisplayRate() : TOptional<FFrameRate>();
-	const FDaySequenceTime NewTimeOfDayPreview = DaySequenceHelpers::HoursToTimecode(InHours, FrameRate);
+	const FDaySequenceTime NewTimeOfDayPreview = FDaySequenceTime::FromHours(InHours);
 	TimeOfDayPreview = NewTimeOfDayPreview;
 	OnTimeOfDayPreviewChangedEvent.Broadcast(InHours);
 	OnTimeOfDayPreviewChanged.Broadcast(InHours);
@@ -569,8 +565,7 @@ void ADaySequenceActor::InitializeRootSequence()
 		RootSequence->SetSequenceFlags(EMovieSceneSequenceFlags::Volatile);
 
 		UMovieScene* RootMovieScene = RootSequence->GetMovieScene();
-		const FFrameRate DisplayRate = RootMovieScene->GetDisplayRate();
-		const float DaySeconds = DaySequenceHelpers::TimecodeToSeconds(TimePerCycle, DisplayRate);
+		const float DaySeconds = TimePerCycle.ToSeconds();
 		const int32 RootDuration = RootMovieScene->GetTickResolution().AsFrameNumber(DaySeconds).Value;
 		RootMovieScene->SetPlaybackRange(0, RootDuration);
 #if WITH_EDITOR
@@ -860,9 +855,9 @@ void ADaySequenceActor::OnSequencePlayerUpdate(const UMovieSceneSequencePlayer&
 	{
 		const FFrameRate FrameRate = RootSequence->GetMovieScene()->GetDisplayRate();
 		const double CurrentTimeSeconds = FQualifiedFrameTime(Time, FrameRate).AsSeconds();
-		const float DayCycleSeconds = DaySequenceHelpers::TimecodeToSeconds(TimePerCycle, FrameRate);
+		const float DayCycleSeconds = TimePerCycle.ToSeconds();
 		const float SequenceRatio = CurrentTimeSeconds / DayCycleSeconds;
-		return DaySequenceHelpers::TimecodeToHours(DayLength, FrameRate) * SequenceRatio;
+		return DayLength.ToHours() * SequenceRatio;
 	};
 	const float CurrentHours = FrameTimeToDayHours(CurrentTime);
 	const float PreviousHours = FrameTimeToDayHours(PreviousTime);
@@ -993,9 +988,8 @@ void ADaySequenceActor::WarpEvaluationRange(FMovieSceneEvaluationRange& InOutRan
 	}
 
 	const FFrameRate TickRate   = InOutRange.GetFrameRate();
-	const FFrameRate FrameRate  = RootSequence->GetMovieScene()->GetDisplayRate();
-	const float DayCycleSeconds = DaySequenceHelpers::TimecodeToSeconds(TimePerCycle, FrameRate);
-	const float DayLengthHours  = DaySequenceHelpers::TimecodeToHours(DayLength, FrameRate);
+	const float DayCycleSeconds = TimePerCycle.ToSeconds();
+	const float DayLengthHours  = DayLength.ToSeconds();
 
 	TRange<FFrameTime> Range = InOutRange.GetRange();
 
@@ -1030,34 +1024,26 @@ void ADaySequenceActor::WarpEvaluationRange(FMovieSceneEvaluationRange& InOutRan
 
 float ADaySequenceActor::GetDayLength() const
 {
-	const TOptional<FFrameRate> FrameRate = RootSequence ? RootSequence->GetMovieScene()->GetDisplayRate() : TOptional<FFrameRate>();
-	return DaySequenceHelpers::TimecodeToHours(DayLength, FrameRate);
+	return DayLength.ToHours();
 }
 
 void ADaySequenceActor::SetDayLength(float InHours)
 {
-	const TOptional<FFrameRate> FrameRate = RootSequence ? RootSequence->GetMovieScene()->GetDisplayRate() : TOptional<FFrameRate>();
-
 	// Set min day length to 1 second.
-	const FDaySequenceTime OneSecond(0, 0, 1);
-	InHours = FMath::Max(InHours, DaySequenceHelpers::TimecodeToHours(OneSecond, FrameRate));
-	DayLength = DaySequenceHelpers::HoursToTimecode(InHours, FrameRate);
+	InHours = FMath::Max(InHours, FDaySequenceTime::FromSeconds(1.f).ToHours());
+	DayLength = FDaySequenceTime::FromHours(InHours);
 }
 
 float ADaySequenceActor::GetTimePerCycle() const
 {
-	const TOptional<FFrameRate> FrameRate = RootSequence ? RootSequence->GetMovieScene()->GetDisplayRate() : TOptional<FFrameRate>();
-	return DaySequenceHelpers::TimecodeToHours(TimePerCycle, FrameRate);
+	return TimePerCycle.ToHours();
 }
 
 void ADaySequenceActor::SetTimePerCycle(float InHours)
 {
-	const TOptional<FFrameRate> FrameRate = RootSequence ? RootSequence->GetMovieScene()->GetDisplayRate() : TOptional<FFrameRate>();
-
 	// Set min cycle length to 1 second.
-	const FDaySequenceTime OneSecond(0, 0, 1);
-	InHours = FMath::Max(InHours, DaySequenceHelpers::TimecodeToHours(OneSecond, FrameRate));
-	TimePerCycle = DaySequenceHelpers::HoursToTimecode(InHours, FrameRate);
+	InHours = FMath::Max(InHours, FDaySequenceTime::FromSeconds(1.f).ToHours());
+	TimePerCycle = FDaySequenceTime::FromHours(InHours);
 }
 
 void ADaySequenceActor::Multicast_SetTimePerCycle_Implementation(float InHours)
@@ -1068,9 +1054,8 @@ void ADaySequenceActor::Multicast_SetTimePerCycle_Implementation(float InHours)
 	const float CurrentTimeOfDay = bIsGameWorld ? GetTimeOfDay() : 0.f;
 
 	// Set min cycle length to 1 second.
-	const FDaySequenceTime OneSecond(0, 0, 1);
-	InHours = FMath::Max(InHours, DaySequenceHelpers::TimecodeToHours(OneSecond, FrameRate));
-	const FDaySequenceTime NewTimePerCycle = DaySequenceHelpers::HoursToTimecode(InHours, FrameRate); 
+	InHours = FMath::Max(InHours, FDaySequenceTime::FromSeconds(1.f).ToHours());
+	const FDaySequenceTime NewTimePerCycle = FDaySequenceTime::FromHours(InHours);
 	if (NewTimePerCycle != TimePerCycle)
 	{
 		// Validate the new TimePerCycle to avoid overflowing the playback range.
@@ -1079,7 +1064,7 @@ void ADaySequenceActor::Multicast_SetTimePerCycle_Implementation(float InHours)
 		const FFrameRate TickResolution = RootMovieScene->GetTickResolution();
 		const auto IsTimePerCycleOverflow = [&DisplayRate, &TickResolution](const FDaySequenceTime& InTimePerCycle)
 		{
-			const float TimePerCycleSeconds = DaySequenceHelpers::TimecodeToSeconds(InTimePerCycle, TOptional<FFrameRate>());
+			const float TimePerCycleSeconds = InTimePerCycle.ToSeconds();
 			const int32 CycleSeconds = static_cast<int32>(TimePerCycleSeconds) + 1;
 			const int32 DisplayFactor = DisplayRate.Numerator / DisplayRate.Denominator;
 			const int32 TickFactor = TickResolution.Numerator / TickResolution.Denominator;
@@ -1098,7 +1083,7 @@ void ADaySequenceActor::Multicast_SetTimePerCycle_Implementation(float InHours)
 		if (bIsGameWorld)
 		{
 			// Update playback range for the root sequence.
-			const float DaySeconds = DaySequenceHelpers::TimecodeToSeconds(TimePerCycle, FFrameRate());
+			const float DaySeconds = TimePerCycle.ToSeconds();
 			const int32 RootDuration = RootMovieScene->GetTickResolution().AsFrameNumber(DaySeconds).Value;
 			RootMovieScene->MarkAsChanged();
 			RootMovieScene->SetPlaybackRange(0, RootDuration);
@@ -1135,15 +1120,13 @@ void ADaySequenceActor::Multicast_SetTimePerCycle_Implementation(float InHours)
 
 float ADaySequenceActor::GetInitialTimeOfDay() const
 {
-	const TOptional<FFrameRate> FrameRate = RootSequence ? RootSequence->GetMovieScene()->GetDisplayRate() : TOptional<FFrameRate>();
-	return DaySequenceHelpers::TimecodeToHours(InitialTimeOfDay, FrameRate);
+	return InitialTimeOfDay.ToHours();
 }
 
 void ADaySequenceActor::SetInitialTimeOfDay(float InHours)
 {
 	InHours = FMath::Clamp(InHours, 0.0f, GetDayLength()); 
-	const TOptional<FFrameRate> FrameRate = RootSequence ? RootSequence->GetMovieScene()->GetDisplayRate() : TOptional<FFrameRate>();
-	InitialTimeOfDay = DaySequenceHelpers::HoursToTimecode(InHours, FrameRate);
+	InitialTimeOfDay = FDaySequenceTime::FromHours(InHours);
 }
 
 float ADaySequenceActor::GetTimeOfDay() const
@@ -1153,12 +1136,11 @@ float ADaySequenceActor::GetTimeOfDay() const
 	const UDaySequencePlayer* Player = GetSequencePlayer();
 	if (HasValidRootSequence() && Player && World && World->IsGameWorld())
 	{
-		const FFrameRate FrameRate = RootSequence->GetMovieScene()->GetDisplayRate();
 		const FQualifiedFrameTime CurrentFrameTime = Player->GetCurrentTime();
 		const double CurrentTimeSeconds = CurrentFrameTime.AsSeconds();
-		const float DayCycleSeconds = DaySequenceHelpers::TimecodeToSeconds(TimePerCycle, FrameRate);
+		const float DayCycleSeconds = TimePerCycle.ToSeconds();
 		const float SequenceRatio = CurrentTimeSeconds / DayCycleSeconds;
-		Result = DaySequenceHelpers::TimecodeToHours(DayLength, FrameRate) * SequenceRatio;
+		Result = DayLength.ToHours() * SequenceRatio;
 	}
 	else
 	{
@@ -1178,7 +1160,7 @@ bool ADaySequenceActor::SetTimeOfDay(float InHours)
 		const FFrameRate FrameRate = RootSequence->GetMovieScene()->GetDisplayRate();
 		const float DayLengthHours = GetDayLength();
 		const float DayLengthRatio = FMath::Frac(InHours / DayLengthHours);
-		const float DayCycleSeconds = DaySequenceHelpers::TimecodeToSeconds(TimePerCycle, FrameRate) * DayLengthRatio;
+		const float DayCycleSeconds = TimePerCycle.ToSeconds() * DayLengthRatio;
 
 		// Update the playback position of the sequence.
 		FMovieSceneSequencePlaybackParams PlaybackParams;
@@ -1231,7 +1213,7 @@ void ADaySequenceActor::SetStaticTimeOfDay(float InHours)
 	const FFrameRate FrameRate = RootSequence->GetMovieScene()->GetDisplayRate();
 	const float DayLengthHours = GetDayLength();
 	const float DayLengthRatio = FMath::Frac(InHours / DayLengthHours);
-	const float DayCycleSeconds = DaySequenceHelpers::TimecodeToSeconds(TimePerCycle, FrameRate) * DayLengthRatio;
+	const float DayCycleSeconds = TimePerCycle.ToSeconds() * DayLengthRatio;
 	const FFrameTime StaticFrameTime = FrameRate.AsFrameTime(DayCycleSeconds);
 
 	if (TSharedPtr<FStaticTimeControllerOverride> ExistingTimeControllerOverride = WeakTimeControllerOverride.Pin())
@@ -1406,7 +1388,6 @@ FMovieSceneSequencePlaybackSettings ADaySequenceActor::GetPlaybackSettings(const
 	{
 		// Convert the initial day time to sequence time
 		const UMovieScene* MovieScene = Sequence->GetMovieScene();
-		const FFrameRate DisplayRate = MovieScene->GetDisplayRate();
 
 #if WITH_EDITORONLY_DATA
 		const float InitialHours = bOverrideInitialTimeOfDay ? GetTimeOfDayPreview() : GetInitialTimeOfDay();
@@ -1416,7 +1397,7 @@ FMovieSceneSequencePlaybackSettings ADaySequenceActor::GetPlaybackSettings(const
 		
 		const float DayLengthHours = GetDayLength();
 		const float StartRatio = FMath::Frac(InitialHours / DayLengthHours);
-		const float DayCycleSeconds = DaySequenceHelpers::TimecodeToSeconds(TimePerCycle, DisplayRate) * StartRatio;
+		const float DayCycleSeconds = TimePerCycle.ToSeconds() * StartRatio;
 		Settings.StartTime = DayCycleSeconds;
 	}
 	return Settings;
diff --git a/Engine/Plugins/Experimental/DaySequence/Source/DaySequence/Private/DaySequenceHelpers.cpp b/Engine/Plugins/Experimental/DaySequence/Source/DaySequence/Private/DaySequenceHelpers.cpp
deleted file mode 100644
diff --git a/Engine/Plugins/Experimental/DaySequence/Source/DaySequence/Private/DaySequenceHelpers.h b/Engine/Plugins/Experimental/DaySequence/Source/DaySequence/Private/DaySequenceHelpers.h
deleted file mode 100644
