DaySequenceActor: Early out for CDO in InitializeRootSequence and make subsequence movie scenes read only if the backing asset is cooked.

#changelist validated
#virtualized

diff --git a/Engine/Plugins/Experimental/DaySequence/Source/DaySequence/Private/DaySequenceActor.cpp b/Engine/Plugins/Experimental/DaySequence/Source/DaySequence/Private/DaySequenceActor.cpp
--- a/Engine/Plugins/Experimental/DaySequence/Source/DaySequence/Private/DaySequenceActor.cpp
+++ b/Engine/Plugins/Experimental/DaySequence/Source/DaySequence/Private/DaySequenceActor.cpp
@@ -727,6 +727,13 @@ UMovieSceneSubSection* ADaySequenceActor::InitializeDaySequence(const FDaySequen
 	UDaySequenceTrack* SubTrack = RootMovieScene->AddTrack<UDaySequenceTrack>();
 	UMovieSceneSubSection* SubSection = nullptr;
 
+#if WITH_EDITORONLY_DATA
+	if (UMovieScene* SequenceMovieScene = SequenceAsset.Sequence->GetMovieScene())
+	{
+		SequenceMovieScene->SetReadOnly(SequenceAsset.Sequence->GetPackage()->bIsCookedForEditor);
+	}
+#endif
+	
 	if (SubTrack)
 	{
 		SubTrack->ClearFlags(RF_Transactional);
