DaySequenceModifierComponent: Lock subsections

diff --git a/Engine/Plugins/Experimental/DaySequence/Source/DaySequence/Private/DaySequenceModifierComponent.cpp b/Engine/Plugins/Experimental/DaySequence/Source/DaySequence/Private/DaySequenceModifierComponent.cpp
--- a/Engine/Plugins/Experimental/DaySequence/Source/DaySequence/Private/DaySequenceModifierComponent.cpp
+++ b/Engine/Plugins/Experimental/DaySequence/Source/DaySequence/Private/DaySequenceModifierComponent.cpp
@@ -838,7 +838,8 @@ UMovieSceneSubSection* UDaySequenceModifierComponent::InitializeDaySequence(cons
 		SubSection->SetSequence(Sequence);
 		SubSection->SetRange(MovieScene->GetPlaybackRange());
 		SubSection->SetIsActive(bActivate);
-
+		SubSection->SetIsLocked(true);
+		
 		TargetActor->UpdateSubSectionTimeScale(SubSection);
 
 		RootTrack->AddSection(*SubSection);
