DaySequence: Temporarily unlock subsections when attempting to change active state

diff --git a/Engine/Plugins/Experimental/DaySequence/Source/DaySequence/Private/DaySequenceActor.cpp b/Engine/Plugins/Experimental/DaySequence/Source/DaySequence/Private/DaySequenceActor.cpp
--- a/Engine/Plugins/Experimental/DaySequence/Source/DaySequence/Private/DaySequenceActor.cpp
+++ b/Engine/Plugins/Experimental/DaySequence/Source/DaySequence/Private/DaySequenceActor.cpp
@@ -760,8 +760,10 @@ UMovieSceneSubSection* ADaySequenceActor::InitializeDaySequence(const FDaySequen
 				constexpr bool bInitialMuteState = false;
 				if (const bool bActive = !EvaluateSequenceConditions(bInitialMuteState, Conditions); SubSection->IsActive() != bActive)
 				{
+					SubSection->SetIsLocked(false);
 					SubSection->MarkAsChanged();
 					SubSection->SetIsActive(bActive);
+					SubSection->SetIsLocked(true);
 				}
 			};
 
@@ -774,8 +776,10 @@ UMovieSceneSubSection* ADaySequenceActor::InitializeDaySequence(const FDaySequen
 
 				if (!SubSection->IsActive())
 				{
+					SubSection->SetIsLocked(false);
 					SubSection->MarkAsChanged();
 					SubSection->SetIsActive(true);
+					SubSection->SetIsLocked(true);
 				}
 			};
 			
@@ -849,9 +853,12 @@ void ADaySequenceActor::UpdateSubSectionTimeScale(UMovieSceneSubSection* InSubSe
 	// normalize playback by setting TimeScale on the section.
 	const UMovieScene* RootMovieScene = RootSequence->GetMovieScene();
 	const int32 RootDuration = RootMovieScene->GetPlaybackRange().GetUpperBoundValue().Value;
+	const bool bSubsectionWasLocked = InSubSection->IsLocked();
+	InSubSection->SetIsLocked(false);
 	InSubSection->MarkAsChanged();
 	InSubSection->Parameters.TimeScale = (float)OuterDuration / (float)RootDuration;
 	InSubSection->SetRange(RootMovieScene->GetPlaybackRange());
+	InSubSection->SetIsLocked(bSubsectionWasLocked);
 }
 
 void ADaySequenceActor::OnSequencePlayerUpdate(const UMovieSceneSequencePlayer& Player, FFrameTime CurrentTime, FFrameTime PreviousTime)
diff --git a/Engine/Plugins/Experimental/DaySequence/Source/DaySequence/Private/DaySequenceModifierComponent.cpp b/Engine/Plugins/Experimental/DaySequence/Source/DaySequence/Private/DaySequenceModifierComponent.cpp
--- a/Engine/Plugins/Experimental/DaySequence/Source/DaySequence/Private/DaySequenceModifierComponent.cpp
+++ b/Engine/Plugins/Experimental/DaySequence/Source/DaySequence/Private/DaySequenceModifierComponent.cpp
@@ -838,7 +838,6 @@ UMovieSceneSubSection* UDaySequenceModifierComponent::InitializeDaySequence(cons
 		SubSection->SetSequence(Sequence);
 		SubSection->SetRange(MovieScene->GetPlaybackRange());
 		SubSection->SetIsActive(bActivate);
-		SubSection->SetIsLocked(true);
 		
 		TargetActor->UpdateSubSectionTimeScale(SubSection);
 
@@ -867,6 +866,7 @@ UMovieSceneSubSection* UDaySequenceModifierComponent::InitializeDaySequence(cons
 #endif
         
 		SubSection->MarkAsChanged();
+		SubSection->SetIsLocked(true);
 		return SubSection;
 	};
 	
@@ -890,8 +890,10 @@ UMovieSceneSubSection* UDaySequenceModifierComponent::InitializeDaySequence(cons
 		const bool bActive = bIsEnabled && !TargetActor->EvaluateSequenceConditions(bInitialMuteState, Conditions);
 		if (SubSection->IsActive() != bActive)
 		{
+			SubSection->SetIsLocked(false);
 			SubSection->MarkAsChanged();
 			SubSection->SetIsActive(bActive);
+			SubSection->SetIsLocked(true);
 		}
 	};
 
@@ -905,8 +907,10 @@ UMovieSceneSubSection* UDaySequenceModifierComponent::InitializeDaySequence(cons
 		const bool bActive = bIsEnabled;
 		if (SubSection->IsActive() != bActive)
 		{
+			SubSection->SetIsLocked(false);
 			SubSection->MarkAsChanged();
 			SubSection->SetIsActive(bActive);
+			SubSection->SetIsLocked(true);
 		}
 	};
 
@@ -1000,8 +1004,10 @@ FGuid UDaySequenceModifierComponent::GetOrCreateProceduralBinding(UObject* Objec
 	UMovieSceneSubSection* SubSection = WeakSubSection.Get();
 	if (SubSection && SubSection->GetSequence() == nullptr)
 	{
+		SubSection->SetIsLocked(false);
 		SubSection->MarkAsChanged();
 		SubSection->SetSequence(ProceduralDaySequence);
+		SubSection->SetIsLocked(true);
 	}
 
 	UMovieScene* MovieScene = ProceduralDaySequence->GetMovieScene();
