DaySequenceActor: Add SequenceUpdateInterval property to the DSA

diff --git a/Engine/Plugins/Experimental/DaySequence/Source/DaySequence/Private/DaySequenceActor.cpp b/Engine/Plugins/Experimental/DaySequence/Source/DaySequence/Private/DaySequenceActor.cpp
--- a/Engine/Plugins/Experimental/DaySequence/Source/DaySequence/Private/DaySequenceActor.cpp
+++ b/Engine/Plugins/Experimental/DaySequence/Source/DaySequence/Private/DaySequenceActor.cpp
@@ -109,6 +109,7 @@ namespace UE::DaySequence
 
 ADaySequenceActor::ADaySequenceActor(const FObjectInitializer& Init)
 : Super(Init)
+, SequenceUpdateInterval(0.f)
 , bRunDayCycle(true)
 , bUseInterpCurve(false)
 , DayLength(24, 0, 0)
diff --git a/Engine/Plugins/Experimental/DaySequence/Source/DaySequence/Public/DaySequenceActor.h b/Engine/Plugins/Experimental/DaySequence/Source/DaySequence/Public/DaySequenceActor.h
--- a/Engine/Plugins/Experimental/DaySequence/Source/DaySequence/Public/DaySequenceActor.h
+++ b/Engine/Plugins/Experimental/DaySequence/Source/DaySequence/Public/DaySequenceActor.h
@@ -474,6 +474,9 @@ protected:
 	UPROPERTY(Transient)
 	TObjectPtr<UDaySequence> RootSequence;
 
+	UPROPERTY(EditAnywhere, Category=RuntimeDayCycle)
+	float SequenceUpdateInterval;
+	
 #if WITH_EDITORONLY_DATA
 	/**
 	 * Sets the time of day to preview in the editor. Does not affect the start time at runtime.
