DaySequenceModifierComponent: Stop assuming PlayerCameraManager is valid

diff --git a/Engine/Plugins/Experimental/DaySequence/Source/DaySequence/Private/DaySequenceModifierComponent.cpp b/Engine/Plugins/Experimental/DaySequence/Source/DaySequence/Private/DaySequenceModifierComponent.cpp
--- a/Engine/Plugins/Experimental/DaySequence/Source/DaySequence/Private/DaySequenceModifierComponent.cpp
+++ b/Engine/Plugins/Experimental/DaySequence/Source/DaySequence/Private/DaySequenceModifierComponent.cpp
@@ -1358,7 +1358,7 @@ bool UDaySequenceModifierComponent::GetBlendPosition(FVector& InPosition) const
 	}
 	else
 #endif
-	if (const APlayerController* BlendTarget = ExternalVolumeBlendTarget.Get())
+	if (const APlayerController* BlendTarget = ExternalVolumeBlendTarget.Get(); BlendTarget && BlendTarget->PlayerCameraManager)
 	{
 		CSV_SCOPED_TIMING_STAT(DaySequence, GetPlayerViewPoint);
 		InPosition = BlendTarget->PlayerCameraManager->GetCameraLocation();
