ProceduralDaySequence: Create a base struct that is an abstraction of a procedural day sequence

This removes AProceduralDaySequenceActor and ASunAngleDaySequenceActor. The former was removed because it was made obsolete, while the latter was removed because it was refactored into a procedural sequence. ASunPositionDaySequenceActor remains because it does more than just procedural sequence generation.

We add the notion of a "procedural sequence" which can be configured inside of a collection asset. Any DSA can use a procedural sequence.

#virtualized

diff --git a/Engine/Plugins/Experimental/DaySequence/Source/DaySequence/Private/Actors/ProceduralDaySequenceActor.cpp b/Engine/Plugins/Experimental/DaySequence/Source/DaySequence/Private/Actors/ProceduralDaySequenceActor.cpp
deleted file mode 100644
diff --git a/Engine/Plugins/Experimental/DaySequence/Source/DaySequence/Private/Actors/SunAngleDaySequenceActor.cpp b/Engine/Plugins/Experimental/DaySequence/Source/DaySequence/Private/Actors/SunAngleDaySequenceActor.cpp
deleted file mode 100644
diff --git a/Engine/Plugins/Experimental/DaySequence/Source/DaySequence/Private/Actors/SunPositionDaySequenceActor.cpp b/Engine/Plugins/Experimental/DaySequence/Source/DaySequence/Private/Actors/SunPositionDaySequenceActor.cpp
--- a/Engine/Plugins/Experimental/DaySequence/Source/DaySequence/Private/Actors/SunPositionDaySequenceActor.cpp
+++ b/Engine/Plugins/Experimental/DaySequence/Source/DaySequence/Private/Actors/SunPositionDaySequenceActor.cpp
@@ -5,170 +5,10 @@
 #include "Components/DirectionalLightComponent.h"
 #include "Components/StaticMeshComponent.h"
 #include "DaySequenceCollectionAsset.h"
-#include "ProceduralDaySequenceBuilder.h"
 
 #include "Materials/MaterialInterface.h"
 #include "UObject/ConstructorHelpers.h"
 
-namespace UE::DaySequence
-{
-	struct FSunPositionData
-	{
-		/** Sun Elevation */
-		double Elevation;
-
-		/** Sun Elevation, corrected for atmospheric diffraction */
-		double CorrectedElevation;
-
-		/** Sun azimuth */
-		double Azimuth;
-
-		/** Sunrise time */
-		FTimespan SunriseTime;
-
-		/** Sunset time */
-		FTimespan SunsetTime;
-
-		/** Solar noon */
-		FTimespan SolarNoon;
-	};
-
-	void GetSunPosition(double Latitude, double Longitude, double TimeZone, bool bIsDaylightSavingTime, int32 Year, int32 Month, int32 Day, int32 Hours, int32 Minutes, int32 Seconds, int32 Milliseconds, FSunPositionData& SunPositionData)
-	{
-		if (!FDateTime::Validate(Year, Month, Day, Hours, Minutes, Seconds, Milliseconds))
-		{
-			return;
-		}
-
-		FDateTime CalcTime(Year, Month, Day, Hours, Minutes, Seconds, Milliseconds);
-
-		double TimeOffset = TimeZone;
-		if (bIsDaylightSavingTime)
-		{
-			TimeOffset += 1.0;
-		}
-
-		double LatitudeRad = FMath::DegreesToRadians(Latitude);
-
-		// Get the julian day (number of days since Jan 1st of the year 4713 BC)
-		double JulianDay = CalcTime.GetJulianDay();
-		double JulianCentury = (JulianDay - 2451545.0) / 36525.0;
-
-		// Get the sun's mean longitude , referred to the mean equinox of julian date
-		double GeomMeanLongSunDeg = FMath::Fmod(280.46646 + JulianCentury*(36000.76983 + JulianCentury*0.0003032), 360.0);
-		double GeomMeanLongSunRad = FMath::DegreesToRadians(GeomMeanLongSunDeg);
-
-		// Get the sun's mean anomaly
-		double GeomMeanAnomSunDeg = 357.52911 + JulianCentury*(35999.05029 - 0.0001537*JulianCentury);
-		double GeomMeanAnomSunRad = FMath::DegreesToRadians(GeomMeanAnomSunDeg);
-
-		// Get the earth's orbit eccentricity
-		double EccentEarthOrbit = 0.016708634 - JulianCentury*(0.000042037 + 0.0000001267*JulianCentury);
-		
-		// Get the sun's equation of the center
-		double SunEqOfCtr = FMath::Sin(GeomMeanAnomSunRad)*(1.914602 - JulianCentury*(0.004817 + 0.000014*JulianCentury))
-						  + FMath::Sin(2.0 * GeomMeanAnomSunRad)*(0.019993 - 0.000101*JulianCentury)
-						  + FMath::Sin(3.0 * GeomMeanAnomSunRad)*0.000289;
-		
-		// Get the sun's true longitude
-		double SunTrueLongDeg = GeomMeanLongSunDeg + SunEqOfCtr;
-
-		// Get the sun's true anomaly
-		//	double SunTrueAnomDeg = GeomMeanAnomSunDeg + SunEqOfCtr;
-		//	double SunTrueAnomRad = FMath::DegreesToRadians(SunTrueAnomDeg);
-
-		// Get the earth's distance from the sun
-		//	double SunRadVectorAUs = (1.000001018*(1.0 - EccentEarthOrbit*EccentEarthOrbit)) / (1.0 + EccentEarthOrbit*FMath::Cos(SunTrueAnomRad));
-		
-		// Get the sun's apparent longitude
-		double SunAppLongDeg = SunTrueLongDeg - 0.00569 - 0.00478*FMath::Sin(FMath::DegreesToRadians(125.04 - 1934.136*JulianCentury));
-		double SunAppLongRad = FMath::DegreesToRadians(SunAppLongDeg);
-
-		// Get the earth's mean obliquity of the ecliptic
-		double MeanObliqEclipticDeg = 23.0 + (26.0 + ((21.448 - JulianCentury*(46.815 + JulianCentury*(0.00059 - JulianCentury*0.001813)))) / 60.0) / 60.0;
-		
-		// Get the oblique correction
-		double ObliqCorrDeg = MeanObliqEclipticDeg + 0.00256*FMath::Cos(FMath::DegreesToRadians(125.04 - 1934.136*JulianCentury));
-		double ObliqCorrRad = FMath::DegreesToRadians(ObliqCorrDeg);
-
-		// Get the sun's right ascension
-		double SunRtAscenRad = FMath::Atan2(FMath::Cos(ObliqCorrRad)*FMath::Sin(SunAppLongRad), FMath::Cos(SunAppLongRad));
-		double SunRtAscenDeg = FMath::RadiansToDegrees(SunRtAscenRad);
-		
-		// Get the sun's declination
-		double SunDeclinRad = FMath::Asin(FMath::Sin(ObliqCorrRad)*FMath::Sin(SunAppLongRad));
-		double SunDeclinDeg = FMath::RadiansToDegrees(SunDeclinRad);
-
-		double VarY = FMath::Pow(FMath::Tan(ObliqCorrRad / 2.0), 2.0);
-		
-		// Get the equation of time
-		double EqOfTimeMinutes = 4.0 * FMath::RadiansToDegrees(VarY*FMath::Sin(2.0 * GeomMeanLongSunRad) - 2.0 * EccentEarthOrbit*FMath::Sin(GeomMeanAnomSunRad) + 4.0 * EccentEarthOrbit*VarY*FMath::Sin(GeomMeanAnomSunRad)*FMath::Cos(2.0 * GeomMeanLongSunRad) - 0.5*VarY*VarY*FMath::Sin(4.0 * GeomMeanLongSunRad) - 1.25*EccentEarthOrbit*EccentEarthOrbit*FMath::Sin(2.0 * GeomMeanAnomSunRad));
-
-		// Get the hour angle of the sunrise
-		double HASunriseDeg = FMath::RadiansToDegrees(FMath::Acos(FMath::Cos(FMath::DegreesToRadians(90.833)) / (FMath::Cos(LatitudeRad)*FMath::Cos(SunDeclinRad)) - FMath::Tan(LatitudeRad)*FMath::Tan(SunDeclinRad)));
-		//	double SunlightDurationMinutes = 8.0 * HASunriseDeg;
-
-		// Get the local time of the sun's rise and set
-		double SolarNoonLST = (720.0 - 4.0 * Longitude - EqOfTimeMinutes + TimeOffset * 60.0) / 1440.0;
-		double SunriseTimeLST = SolarNoonLST - HASunriseDeg * 4.0 / 1440.0;
-		double SunsetTimeLST = SolarNoonLST + HASunriseDeg * 4.0 / 1440.0;
-
-		// Get the true solar time
-		double TrueSolarTimeMinutes = FMath::Fmod(CalcTime.GetTimeOfDay().GetTotalMinutes() + EqOfTimeMinutes + 4.0 * Longitude - 60.0 * TimeOffset, 1440.0);
-
-		// Get the hour angle of current time
-		double HourAngleDeg = TrueSolarTimeMinutes < 0 ? TrueSolarTimeMinutes / 4.0 + 180 : TrueSolarTimeMinutes / 4.0 - 180.0;
-		double HourAngleRad = FMath::DegreesToRadians(HourAngleDeg);
-
-		// Get the solar zenith angle
-		double SolarZenithAngleRad = FMath::Acos(FMath::Sin(LatitudeRad)*FMath::Sin(SunDeclinRad) + FMath::Cos(LatitudeRad)*FMath::Cos(SunDeclinRad)*FMath::Cos(HourAngleRad));
-		double SolarZenithAngleDeg = FMath::RadiansToDegrees(SolarZenithAngleRad);
-
-		// Get the sun elevation
-		double SolarElevationAngleDeg = 90.0 - SolarZenithAngleDeg;
-		double SolarElevationAngleRad = FMath::DegreesToRadians(SolarElevationAngleDeg);
-		double TanOfSolarElevationAngle = FMath::Tan(SolarElevationAngleRad);
-
-		// Get the approximated atmospheric refraction
-		double ApproxAtmosphericRefractionDeg = 0.0;
-		if (SolarElevationAngleDeg <= 85.0)
-		{
-			if (SolarElevationAngleDeg > 5.0)
-			{
-				ApproxAtmosphericRefractionDeg = 58.1 / TanOfSolarElevationAngle - 0.07 / FMath::Pow(TanOfSolarElevationAngle, 3) + 0.000086 / FMath::Pow(TanOfSolarElevationAngle, 5) / 3600.0;
-			}
-			else
-			{
-				if (SolarElevationAngleDeg > -0.575)
-				{
-					ApproxAtmosphericRefractionDeg = 1735.0 + SolarElevationAngleDeg * (-518.2 + SolarElevationAngleDeg * (103.4 + SolarElevationAngleDeg * (-12.79 + SolarElevationAngleDeg * 0.711)));
-				}
-				else
-				{
-					ApproxAtmosphericRefractionDeg = -20.772 / TanOfSolarElevationAngle;
-				}
-			}
-			ApproxAtmosphericRefractionDeg /= 3600.0;
-		}
-
-		// Get the corrected solar elevation
-		double SolarElevationcorrectedforatmrefractionDeg = SolarElevationAngleDeg + ApproxAtmosphericRefractionDeg;
-
-		// Get the solar azimuth 
-		double tmp = FMath::RadiansToDegrees(FMath::Acos(((FMath::Sin(LatitudeRad)*FMath::Cos(SolarZenithAngleRad)) - FMath::Sin(SunDeclinRad)) / (FMath::Cos(LatitudeRad)*FMath::Sin(SolarZenithAngleRad))));
-		double SolarAzimuthAngleDegcwfromN = HourAngleDeg > 0.0 ? FMath::Fmod(tmp + 180.0, 360.0) : FMath::Fmod(540.0 - tmp, 360.0);
-
-
-		// offset elevation angle to fit with UE coords system
-		SunPositionData.Elevation = 180.0 + SolarElevationAngleDeg;
-		SunPositionData.CorrectedElevation = 180.0 + SolarElevationcorrectedforatmrefractionDeg;
-		SunPositionData.Azimuth = SolarAzimuthAngleDegcwfromN;
-		SunPositionData.SolarNoon = FTimespan::FromDays(SolarNoonLST);
-		SunPositionData.SunriseTime = FTimespan::FromDays(SunriseTimeLST);
-		SunPositionData.SunsetTime = FTimespan::FromDays(SunsetTimeLST);
-	}
-}
-
 ASunPositionDaySequenceActor::ASunPositionDaySequenceActor(const FObjectInitializer& Init)
 : Super(Init)
 {
@@ -192,80 +32,4 @@ ASunPositionDaySequenceActor::ASunPositionDaySequenceActor(const FObjectInitiali
 	// Override the default collection (which animates the moon and sky material)
 	static ConstructorHelpers::FObjectFinder<UDaySequenceCollectionAsset> DefaultCollection(TEXT("/DaySequence/DSCA_24hr.DSCA_24hr"));
 	DaySequenceCollection = DefaultCollection.Object.Get();
-}
-
-void ASunPositionDaySequenceActor::Tick(float DeltaTime)
-{
-	Super::Tick(DeltaTime);
-
-#if WITH_EDITOR
-	ApplySunAngle();
-#endif
-}
-
-void ASunPositionDaySequenceActor::ApplySunAngle()
-{
-	SunComponent->SetRelativeRotation(FRotator(SunPitch, SunYaw, 0.0));
-
-	Azimuth = SunYaw;
-	Elevation = SunPitch - 180.0;
-}
-
-void ASunPositionDaySequenceActor::BuildSequence(UProceduralDaySequenceBuilder* SequenceBuilder)
-{
-	using namespace UE::DaySequence;
-	
-	SequenceBuilder->SetActiveBoundObject(this);
-	
-	const double NormalizedTimeIncrement = 1.0 / FMath::Max(KeyCount - 1, static_cast<unsigned>(1));
-	for (unsigned int Key = 0; Key < KeyCount; ++Key)
-	{
-		const double KeyTime = Key * NormalizedTimeIncrement;
-		
-		FSunPositionData Data;
-
-		double Hours = (double)Key / (double)KeyCount * 24.0;
-		int32 IntHours = Hours;
-
-		double Minutes = FMath::Frac(Hours) * 60.0;
-		int32 IntMinutes = Minutes;
-
-		double Seconds = FMath::Frac(Minutes) * 60.0;
-		int32 IntSeconds = Seconds;
-
-		double Milliseconds = FMath::Frac(Seconds) * 1000.0;
-		int32 IntMilliseconds = Milliseconds;
-		
-		GetSunPosition(Latitude, Longitude, TimeZone, bIsDaylightSavings, Time.GetYear(), Time.GetMonth(), Time.GetDay(), IntHours, IntMinutes, IntSeconds, IntMilliseconds, Data);
-		
-		SequenceBuilder->AddScalarKey("SunPitch", KeyTime, Data.CorrectedElevation);
-		SequenceBuilder->AddScalarKey("SunYaw", KeyTime, Data.Azimuth);
-	}
-}
-
-void ASunPositionDaySequenceActor::SequencePlayerUpdated(float CurrentTime, float PreviousTime)
-{
-	ApplySunAngle();
-}
-
-#if WITH_EDITOR
-void ASunPositionDaySequenceActor::PostEditChangeProperty(FPropertyChangedEvent& PropertyChangedEvent)
-{
-	const FName PropertyName = PropertyChangedEvent.Property ? PropertyChangedEvent.Property->GetFName() : NAME_None;
-
-	// When we finish a drag setting a property that changes the procedural sequence, mark the sequence invalid.
-	// We don't do this during the drag because it will do a _lot_ of work.
-	if (PropertyChangedEvent.ChangeType == EPropertyChangeType::ValueSet &&
-		(PropertyName == GET_MEMBER_NAME_CHECKED(ASunPositionDaySequenceActor, KeyCount) ||
-		 PropertyName == GET_MEMBER_NAME_CHECKED(ASunPositionDaySequenceActor, Time) ||
-		 PropertyName == GET_MEMBER_NAME_CHECKED(ASunPositionDaySequenceActor, TimeZone) ||
-		 PropertyName == GET_MEMBER_NAME_CHECKED(ASunPositionDaySequenceActor, Latitude) ||
-		 PropertyName == GET_MEMBER_NAME_CHECKED(ASunPositionDaySequenceActor, Longitude) ||
-		 PropertyName == GET_MEMBER_NAME_CHECKED(ASunPositionDaySequenceActor, bIsDaylightSavings)))
-	{
-		bProceduralSequenceInvalid = true;
-	}
-
-	Super::PostEditChangeProperty(PropertyChangedEvent);
-}
-#endif
\ No newline at end of file
+}
\ No newline at end of file
diff --git a/Engine/Plugins/Experimental/DaySequence/Source/DaySequence/Private/DaySequenceActor.cpp b/Engine/Plugins/Experimental/DaySequence/Source/DaySequence/Private/DaySequenceActor.cpp
--- a/Engine/Plugins/Experimental/DaySequence/Source/DaySequence/Private/DaySequenceActor.cpp
+++ b/Engine/Plugins/Experimental/DaySequence/Source/DaySequence/Private/DaySequenceActor.cpp
@@ -160,7 +160,7 @@ ADaySequenceActor::ADaySequenceActor(const FObjectInitializer& Init)
 	bReplicateUsingRegisteredSubObjectList = true;
 	bReplicatePlayback = false;
 	SetHidden(false);
-
+	
 #if ENABLE_DRAW_DEBUG
 	if (!IsTemplate())
 	{
@@ -714,6 +714,16 @@ void ADaySequenceActor::InitializeDaySequences()
 		{
 			InitializeDaySequence(Entry);
 		}
+
+		for (TInstancedStruct<FProceduralDaySequence>& ProceduralDaySequence : DaySequenceCollection->ProceduralDaySequences)
+		{
+			if (!ProceduralDaySequence.IsValid())
+			{
+				continue;
+			}
+			
+			InitializeDaySequence(ProceduralDaySequence.GetMutable<FProceduralDaySequence>().GetSequence(this));
+		}
 	}
 }
 
diff --git a/Engine/Plugins/Experimental/DaySequence/Source/DaySequence/Private/ProceduralDaySequence.cpp b/Engine/Plugins/Experimental/DaySequence/Source/DaySequence/Private/ProceduralDaySequence.cpp
new file mode 100644
--- /dev/null
+++ b/Engine/Plugins/Experimental/DaySequence/Source/DaySequence/Private/ProceduralDaySequence.cpp
@@ -0,0 +1,27 @@
+// Copyright Epic Games, Inc. All Rights Reserved.
+
+#include "ProceduralDaySequence.h"
+
+#include "DaySequence.h"
+#include "DaySequenceActor.h"
+
+#include UE_INLINE_GENERATED_CPP_BY_NAME(ProceduralDaySequence)
+
+#define LOCTEXT_NAMESPACE "ProceduralDaySequence"
+
+UDaySequence* FProceduralDaySequence::GetSequence(ADaySequenceActor* InActor)
+{
+	WeakTargetActor = InActor;
+	ProceduralSequence = nullptr;
+	
+	if (!ProceduralSequence && InActor)
+	{
+		UProceduralDaySequenceBuilder* SequenceBuilder = NewObject<UProceduralDaySequenceBuilder>();
+		ProceduralSequence = SequenceBuilder->Initialize(InActor, ProceduralSequence);
+		BuildSequence(SequenceBuilder);
+	}
+	
+	return ProceduralSequence;
+}
+
+#undef LOCTEXT_NAMESPACE
\ No newline at end of file
diff --git a/Engine/Plugins/Experimental/DaySequence/Source/DaySequence/Private/ProceduralDaySequenceBuilder.cpp b/Engine/Plugins/Experimental/DaySequence/Source/DaySequence/Private/ProceduralDaySequenceBuilder.cpp
--- a/Engine/Plugins/Experimental/DaySequence/Source/DaySequence/Private/ProceduralDaySequenceBuilder.cpp
+++ b/Engine/Plugins/Experimental/DaySequence/Source/DaySequence/Private/ProceduralDaySequenceBuilder.cpp
@@ -2,6 +2,7 @@
 
 #include "ProceduralDaySequenceBuilder.h"
 
+#include "DaySequenceActor.h"
 #include "DaySequenceTime.h"
 #include "MovieSceneCommonHelpers.h"
 
@@ -12,6 +13,7 @@
 #include "Tracks/MovieSceneBoolTrack.h"
 #include "Tracks/MovieSceneDoubleTrack.h"
 #include "Tracks/MovieSceneFloatTrack.h"
+#include "Tracks/MovieScene3DTransformTrack.h"
 
 #include "MovieScene.h"
 #include "DaySequence.h"
@@ -51,6 +53,38 @@ namespace UE::DaySequence
 
 		return true;
 	}
+
+	void AddDoubleKey(FFrameNumber Time, double Value, FMovieSceneDoubleChannel* Channel, ERichCurveInterpMode InterpMode)
+	{
+		if (!Channel)
+		{
+			return;
+		}
+		
+		switch (InterpMode)
+		{
+		case RCIM_Linear: Channel->AddLinearKey(Time, Value); break;
+		case RCIM_Constant: Channel->AddConstantKey(Time, Value); break;
+		case RCIM_Cubic: Channel->AddCubicKey(Time, Value); break;
+		case RCIM_None: break;
+		}
+	}
+
+	void AddFloatKey(FFrameNumber Time, double Value, FMovieSceneFloatChannel* Channel, ERichCurveInterpMode InterpMode)
+	{
+		if (!Channel)
+		{
+			return;
+		}
+		
+		switch (InterpMode)
+		{
+		case RCIM_Linear: Channel->AddLinearKey(Time, Value); break;
+		case RCIM_Constant: Channel->AddConstantKey(Time, Value); break;
+		case RCIM_Cubic: Channel->AddCubicKey(Time, Value); break;
+		case RCIM_None: break;
+		}
+	}
 };
 
 UDaySequence* UProceduralDaySequenceBuilder::Initialize(ADaySequenceActor* InActor, UDaySequence* InitialSequence, bool bClearInitialSequence)
@@ -181,24 +215,7 @@ void UProceduralDaySequenceBuilder::AddScalarKeys(FName PropertyName, const TArr
 		for (const TPair<float, double>& KeyValue : KeysAndValues)
 		{
 			const FFrameNumber FrameNumber = UE::DaySequence::GetKeyFrameNumber(KeyValue.Key, MovieScene->GetPlaybackRange());
-
-			switch (InterpMode)
-			{
-			case RCIM_Linear:
-				FloatSection->GetChannel().AddLinearKey(FrameNumber, KeyValue.Value);
-				break;
-				
-			case RCIM_Constant:
-				FloatSection->GetChannel().AddConstantKey(FrameNumber, KeyValue.Value);
-				break;
-			
-			case RCIM_Cubic:
-				FloatSection->GetChannel().AddCubicKey(FrameNumber, KeyValue.Value);
-				break;
-				
-			case RCIM_None:
-				break;
-			}
+			UE::DaySequence::AddFloatKey(FrameNumber, KeyValue.Value, &FloatSection->GetChannel(), InterpMode);
 		}
 
 		FloatSection->MarkAsChanged();
@@ -210,24 +227,7 @@ void UProceduralDaySequenceBuilder::AddScalarKeys(FName PropertyName, const TArr
 		for (const TPair<float, double>& KeyValue : KeysAndValues)
 		{
 			const FFrameNumber FrameNumber = UE::DaySequence::GetKeyFrameNumber(KeyValue.Key, MovieScene->GetPlaybackRange());
-			
-			switch (InterpMode)
-			{
-			case RCIM_Linear:
-				DoubleSection->GetChannel().AddLinearKey(FrameNumber, KeyValue.Value);
-				break;
-				
-			case RCIM_Constant:
-				DoubleSection->GetChannel().AddConstantKey(FrameNumber, KeyValue.Value);
-				break;
-			
-			case RCIM_Cubic:
-				DoubleSection->GetChannel().AddCubicKey(FrameNumber, KeyValue.Value);
-				break;
-				
-			case RCIM_None:
-				break;
-			}
+			UE::DaySequence::AddDoubleKey(FrameNumber, KeyValue.Value, &DoubleSection->GetChannel(), InterpMode);
 		}
 
 		DoubleSection->MarkAsChanged();
@@ -321,6 +321,78 @@ void UProceduralDaySequenceBuilder::AddBoolKeys(FName PropertyName, const TArray
 	}
 }
 
+void UProceduralDaySequenceBuilder::AddTransformKey(float Key, const FTransform& Value, ERichCurveInterpMode InterpMode)
+{
+	using namespace UE::DaySequence;
+	
+	if (!IsInitialized())
+	{
+		FFrame::KismetExecutionMessage(TEXT("AddTransformKey(s) called on an uninitialized Procedural Day Sequence Builder!"), ELogVerbosity::Error);
+		return;
+	}
+	
+	AddTranslationKey(Key, Value.GetLocation(), InterpMode);
+	AddRotationKey(Key, Value.Rotator(), InterpMode);
+	AddScaleKey(Key, Value.GetScale3D(), InterpMode);
+}
+
+void UProceduralDaySequenceBuilder::AddTranslationKey(float Key, const FVector& Value, ERichCurveInterpMode InterpMode)
+{
+	using namespace UE::DaySequence;
+	
+	if (!IsInitialized())
+	{
+		FFrame::KismetExecutionMessage(TEXT("AddTranslationKey(s) called on an uninitialized Procedural Day Sequence Builder!"), ELogVerbosity::Error);
+		return;
+	}
+
+	const UMovieScene* MovieScene = ProceduralDaySequence->GetMovieScene();
+	const FFrameNumber FrameNumber = GetKeyFrameNumber(Key, MovieScene->GetPlaybackRange());
+	const UMovieScene3DTransformSection* TransformSection = CreateOrAddPropertyOverrideSection<UMovieScene3DTransformTrack, UMovieScene3DTransformSection>("Transform");
+
+	AddDoubleKey(FrameNumber, Value.X, TransformSection->GetChannelProxy().GetChannel<FMovieSceneDoubleChannel>(0), InterpMode);
+	AddDoubleKey(FrameNumber, Value.Y, TransformSection->GetChannelProxy().GetChannel<FMovieSceneDoubleChannel>(1), InterpMode);
+	AddDoubleKey(FrameNumber, Value.Z, TransformSection->GetChannelProxy().GetChannel<FMovieSceneDoubleChannel>(2), InterpMode);
+}
+
+void UProceduralDaySequenceBuilder::AddRotationKey(float Key, const FRotator& Value, ERichCurveInterpMode InterpMode)
+{
+	using namespace UE::DaySequence;
+	
+	if (!IsInitialized())
+	{
+		FFrame::KismetExecutionMessage(TEXT("AddRotationKey(s) called on an uninitialized Procedural Day Sequence Builder!"), ELogVerbosity::Error);
+		return;
+	}
+
+	const UMovieScene* MovieScene = ProceduralDaySequence->GetMovieScene();
+	const FFrameNumber FrameNumber = GetKeyFrameNumber(Key, MovieScene->GetPlaybackRange());
+	const UMovieScene3DTransformSection* TransformSection = CreateOrAddPropertyOverrideSection<UMovieScene3DTransformTrack, UMovieScene3DTransformSection>("Transform");
+	
+	AddDoubleKey(FrameNumber, Value.Roll,  TransformSection->GetChannelProxy().GetChannel<FMovieSceneDoubleChannel>(3), InterpMode);
+	AddDoubleKey(FrameNumber, Value.Pitch, TransformSection->GetChannelProxy().GetChannel<FMovieSceneDoubleChannel>(4), InterpMode);
+	AddDoubleKey(FrameNumber, Value.Yaw,   TransformSection->GetChannelProxy().GetChannel<FMovieSceneDoubleChannel>(5), InterpMode);
+}
+
+void UProceduralDaySequenceBuilder::AddScaleKey(float Key, const FVector& Value, ERichCurveInterpMode InterpMode)
+{
+	using namespace UE::DaySequence;
+	
+	if (!IsInitialized())
+	{
+		FFrame::KismetExecutionMessage(TEXT("AddScaleKey(s) called on an uninitialized Procedural Day Sequence Builder!"), ELogVerbosity::Error);
+		return;
+	}
+
+	const UMovieScene* MovieScene = ProceduralDaySequence->GetMovieScene();
+	const FFrameNumber FrameNumber = GetKeyFrameNumber(Key, MovieScene->GetPlaybackRange());
+	const UMovieScene3DTransformSection* TransformSection = CreateOrAddPropertyOverrideSection<UMovieScene3DTransformTrack, UMovieScene3DTransformSection>("Transform");
+
+	AddDoubleKey(FrameNumber, Value.X, TransformSection->GetChannelProxy().GetChannel<FMovieSceneDoubleChannel>(6), InterpMode);
+	AddDoubleKey(FrameNumber, Value.Y, TransformSection->GetChannelProxy().GetChannel<FMovieSceneDoubleChannel>(7), InterpMode);
+	AddDoubleKey(FrameNumber, Value.Z, TransformSection->GetChannelProxy().GetChannel<FMovieSceneDoubleChannel>(8), InterpMode);
+}
+
 void UProceduralDaySequenceBuilder::AddStaticTime(float StaticTime)
 {
 	if (!IsInitialized())
diff --git a/Engine/Plugins/Experimental/DaySequence/Source/DaySequence/Private/ProceduralSequences/SineSequence.cpp b/Engine/Plugins/Experimental/DaySequence/Source/DaySequence/Private/ProceduralSequences/SineSequence.cpp
new file mode 100644
--- /dev/null
+++ b/Engine/Plugins/Experimental/DaySequence/Source/DaySequence/Private/ProceduralSequences/SineSequence.cpp
@@ -0,0 +1,42 @@
+// Copyright Epic Games, Inc. All Rights Reserved.
+
+#include "ProceduralSequences/SineSequence.h"
+
+#include "DaySequence.h"
+#include "DaySequenceActor.h"
+
+#include UE_INLINE_GENERATED_CPP_BY_NAME(SineSequence)
+
+UDaySequence* FSineSequence::BuildSequence(UProceduralDaySequenceBuilder* InBuilder)
+{
+	using namespace UE::DaySequence;
+
+	ADaySequenceActor* TargetActor = WeakTargetActor.Get();
+	
+	if (!TargetActor || PropertyName.IsNone())
+	{
+		return nullptr;
+	}
+
+	UObject* AnimatedObject = TargetActor;
+
+	if (!ComponentName.IsNone())
+	{
+		AnimatedObject = GetComponentByName<USceneComponent>(TargetActor, ComponentName);
+	}
+	
+	if (AnimatedObject)
+	{
+		InBuilder->SetActiveBoundObject(AnimatedObject);
+
+		const double NormalizedTimeIncrement = 1.0 / FMath::Max(KeyCount - 1, static_cast<unsigned>(1));
+		for (unsigned int Key = 0; Key < KeyCount; ++Key)
+		{
+			const double KeyTime = Key * NormalizedTimeIncrement;
+
+			InBuilder->AddScalarKey(PropertyName, KeyTime, Amplitude * FMath::Sin(UE_DOUBLE_TWO_PI * Frequency * (KeyTime - PhaseShift)) + VerticalShift);
+		}
+	}
+	
+	return ProceduralSequence;
+}
\ No newline at end of file
diff --git a/Engine/Plugins/Experimental/DaySequence/Source/DaySequence/Private/ProceduralSequences/SunAngleSequence.cpp b/Engine/Plugins/Experimental/DaySequence/Source/DaySequence/Private/ProceduralSequences/SunAngleSequence.cpp
new file mode 100644
--- /dev/null
+++ b/Engine/Plugins/Experimental/DaySequence/Source/DaySequence/Private/ProceduralSequences/SunAngleSequence.cpp
@@ -0,0 +1,32 @@
+// Copyright Epic Games, Inc. All Rights Reserved.
+
+#include "ProceduralSequences/SunAngleSequence.h"
+
+#include "DaySequence.h"
+#include "DaySequenceActor.h"
+
+#include "Components/DirectionalLightComponent.h"
+
+#include UE_INLINE_GENERATED_CPP_BY_NAME(SunAngleSequence)
+
+UDaySequence* FSunAngleSequence::BuildSequence(UProceduralDaySequenceBuilder* InBuilder)
+{
+	using namespace UE::DaySequence;
+
+	ADaySequenceActor* TargetActor = WeakTargetActor.Get();
+	
+	if (!TargetActor)
+	{
+		return nullptr;
+	}
+	
+	if (UDirectionalLightComponent* SunComponent = GetComponentByName<UDirectionalLightComponent>(TargetActor, SunComponentName))
+	{
+		InBuilder->SetActiveBoundObject(SunComponent);
+
+		InBuilder->AddRotationKey(0.0, FRotator(90.0, 0.0, 0.0), RCIM_Linear);
+		InBuilder->AddRotationKey(1.0, FRotator(450.0, 0.0, 0.0), RCIM_Linear);
+	}
+	
+	return ProceduralSequence;
+}
\ No newline at end of file
diff --git a/Engine/Plugins/Experimental/DaySequence/Source/DaySequence/Private/ProceduralSequences/SunPositionSequence.cpp b/Engine/Plugins/Experimental/DaySequence/Source/DaySequence/Private/ProceduralSequences/SunPositionSequence.cpp
new file mode 100644
--- /dev/null
+++ b/Engine/Plugins/Experimental/DaySequence/Source/DaySequence/Private/ProceduralSequences/SunPositionSequence.cpp
@@ -0,0 +1,212 @@
+// Copyright Epic Games, Inc. All Rights Reserved.
+
+#include "ProceduralSequences/SunPositionSequence.h"
+
+#include "DaySequence.h"
+#include "DaySequenceActor.h"
+
+#include "Components/DirectionalLightComponent.h"
+
+#include UE_INLINE_GENERATED_CPP_BY_NAME(SunPositionSequence)
+
+namespace UE::DaySequence
+{
+	struct FSunPositionData
+    {
+    	/** Sun Elevation */
+    	double Elevation;
+
+    	/** Sun Elevation, corrected for atmospheric diffraction */
+    	double CorrectedElevation;
+
+    	/** Sun azimuth */
+    	double Azimuth;
+
+    	/** Sunrise time */
+    	FTimespan SunriseTime;
+
+    	/** Sunset time */
+    	FTimespan SunsetTime;
+
+    	/** Solar noon */
+    	FTimespan SolarNoon;
+    };
+	
+	void GetSunPosition(double Latitude, double Longitude, double TimeZone, bool bIsDaylightSavingTime, int32 Year, int32 Month, int32 Day, int32 Hours, int32 Minutes, int32 Seconds, int32 Milliseconds, FSunPositionData& SunPositionData)
+	{
+		if (!FDateTime::Validate(Year, Month, Day, Hours, Minutes, Seconds, Milliseconds))
+		{
+			return;
+		}
+
+		FDateTime CalcTime(Year, Month, Day, Hours, Minutes, Seconds, Milliseconds);
+
+		double TimeOffset = TimeZone;
+		if (bIsDaylightSavingTime)
+		{
+			TimeOffset += 1.0;
+		}
+
+		double LatitudeRad = FMath::DegreesToRadians(Latitude);
+
+		// Get the julian day (number of days since Jan 1st of the year 4713 BC)
+		double JulianDay = CalcTime.GetJulianDay();
+		double JulianCentury = (JulianDay - 2451545.0) / 36525.0;
+
+		// Get the sun's mean longitude , referred to the mean equinox of julian date
+		double GeomMeanLongSunDeg = FMath::Fmod(280.46646 + JulianCentury*(36000.76983 + JulianCentury*0.0003032), 360.0);
+		double GeomMeanLongSunRad = FMath::DegreesToRadians(GeomMeanLongSunDeg);
+
+		// Get the sun's mean anomaly
+		double GeomMeanAnomSunDeg = 357.52911 + JulianCentury*(35999.05029 - 0.0001537*JulianCentury);
+		double GeomMeanAnomSunRad = FMath::DegreesToRadians(GeomMeanAnomSunDeg);
+
+		// Get the earth's orbit eccentricity
+		double EccentEarthOrbit = 0.016708634 - JulianCentury*(0.000042037 + 0.0000001267*JulianCentury);
+		
+		// Get the sun's equation of the center
+		double SunEqOfCtr = FMath::Sin(GeomMeanAnomSunRad)*(1.914602 - JulianCentury*(0.004817 + 0.000014*JulianCentury))
+						  + FMath::Sin(2.0 * GeomMeanAnomSunRad)*(0.019993 - 0.000101*JulianCentury)
+						  + FMath::Sin(3.0 * GeomMeanAnomSunRad)*0.000289;
+		
+		// Get the sun's true longitude
+		double SunTrueLongDeg = GeomMeanLongSunDeg + SunEqOfCtr;
+
+		// Get the sun's true anomaly
+		//	double SunTrueAnomDeg = GeomMeanAnomSunDeg + SunEqOfCtr;
+		//	double SunTrueAnomRad = FMath::DegreesToRadians(SunTrueAnomDeg);
+
+		// Get the earth's distance from the sun
+		//	double SunRadVectorAUs = (1.000001018*(1.0 - EccentEarthOrbit*EccentEarthOrbit)) / (1.0 + EccentEarthOrbit*FMath::Cos(SunTrueAnomRad));
+		
+		// Get the sun's apparent longitude
+		double SunAppLongDeg = SunTrueLongDeg - 0.00569 - 0.00478*FMath::Sin(FMath::DegreesToRadians(125.04 - 1934.136*JulianCentury));
+		double SunAppLongRad = FMath::DegreesToRadians(SunAppLongDeg);
+
+		// Get the earth's mean obliquity of the ecliptic
+		double MeanObliqEclipticDeg = 23.0 + (26.0 + ((21.448 - JulianCentury*(46.815 + JulianCentury*(0.00059 - JulianCentury*0.001813)))) / 60.0) / 60.0;
+		
+		// Get the oblique correction
+		double ObliqCorrDeg = MeanObliqEclipticDeg + 0.00256*FMath::Cos(FMath::DegreesToRadians(125.04 - 1934.136*JulianCentury));
+		double ObliqCorrRad = FMath::DegreesToRadians(ObliqCorrDeg);
+
+		// Get the sun's right ascension
+		double SunRtAscenRad = FMath::Atan2(FMath::Cos(ObliqCorrRad)*FMath::Sin(SunAppLongRad), FMath::Cos(SunAppLongRad));
+		double SunRtAscenDeg = FMath::RadiansToDegrees(SunRtAscenRad);
+		
+		// Get the sun's declination
+		double SunDeclinRad = FMath::Asin(FMath::Sin(ObliqCorrRad)*FMath::Sin(SunAppLongRad));
+		double SunDeclinDeg = FMath::RadiansToDegrees(SunDeclinRad);
+
+		double VarY = FMath::Pow(FMath::Tan(ObliqCorrRad / 2.0), 2.0);
+		
+		// Get the equation of time
+		double EqOfTimeMinutes = 4.0 * FMath::RadiansToDegrees(VarY*FMath::Sin(2.0 * GeomMeanLongSunRad) - 2.0 * EccentEarthOrbit*FMath::Sin(GeomMeanAnomSunRad) + 4.0 * EccentEarthOrbit*VarY*FMath::Sin(GeomMeanAnomSunRad)*FMath::Cos(2.0 * GeomMeanLongSunRad) - 0.5*VarY*VarY*FMath::Sin(4.0 * GeomMeanLongSunRad) - 1.25*EccentEarthOrbit*EccentEarthOrbit*FMath::Sin(2.0 * GeomMeanAnomSunRad));
+
+		// Get the hour angle of the sunrise
+		double HASunriseDeg = FMath::RadiansToDegrees(FMath::Acos(FMath::Cos(FMath::DegreesToRadians(90.833)) / (FMath::Cos(LatitudeRad)*FMath::Cos(SunDeclinRad)) - FMath::Tan(LatitudeRad)*FMath::Tan(SunDeclinRad)));
+		//	double SunlightDurationMinutes = 8.0 * HASunriseDeg;
+
+		// Get the local time of the sun's rise and set
+		double SolarNoonLST = (720.0 - 4.0 * Longitude - EqOfTimeMinutes + TimeOffset * 60.0) / 1440.0;
+		double SunriseTimeLST = SolarNoonLST - HASunriseDeg * 4.0 / 1440.0;
+		double SunsetTimeLST = SolarNoonLST + HASunriseDeg * 4.0 / 1440.0;
+
+		// Get the true solar time
+		double TrueSolarTimeMinutes = FMath::Fmod(CalcTime.GetTimeOfDay().GetTotalMinutes() + EqOfTimeMinutes + 4.0 * Longitude - 60.0 * TimeOffset, 1440.0);
+
+		// Get the hour angle of current time
+		double HourAngleDeg = TrueSolarTimeMinutes < 0 ? TrueSolarTimeMinutes / 4.0 + 180 : TrueSolarTimeMinutes / 4.0 - 180.0;
+		double HourAngleRad = FMath::DegreesToRadians(HourAngleDeg);
+
+		// Get the solar zenith angle
+		double SolarZenithAngleRad = FMath::Acos(FMath::Sin(LatitudeRad)*FMath::Sin(SunDeclinRad) + FMath::Cos(LatitudeRad)*FMath::Cos(SunDeclinRad)*FMath::Cos(HourAngleRad));
+		double SolarZenithAngleDeg = FMath::RadiansToDegrees(SolarZenithAngleRad);
+
+		// Get the sun elevation
+		double SolarElevationAngleDeg = 90.0 - SolarZenithAngleDeg;
+		double SolarElevationAngleRad = FMath::DegreesToRadians(SolarElevationAngleDeg);
+		double TanOfSolarElevationAngle = FMath::Tan(SolarElevationAngleRad);
+
+		// Get the approximated atmospheric refraction
+		double ApproxAtmosphericRefractionDeg = 0.0;
+		if (SolarElevationAngleDeg <= 85.0)
+		{
+			if (SolarElevationAngleDeg > 5.0)
+			{
+				ApproxAtmosphericRefractionDeg = 58.1 / TanOfSolarElevationAngle - 0.07 / FMath::Pow(TanOfSolarElevationAngle, 3) + 0.000086 / FMath::Pow(TanOfSolarElevationAngle, 5) / 3600.0;
+			}
+			else
+			{
+				if (SolarElevationAngleDeg > -0.575)
+				{
+					ApproxAtmosphericRefractionDeg = 1735.0 + SolarElevationAngleDeg * (-518.2 + SolarElevationAngleDeg * (103.4 + SolarElevationAngleDeg * (-12.79 + SolarElevationAngleDeg * 0.711)));
+				}
+				else
+				{
+					ApproxAtmosphericRefractionDeg = -20.772 / TanOfSolarElevationAngle;
+				}
+			}
+			ApproxAtmosphericRefractionDeg /= 3600.0;
+		}
+
+		// Get the corrected solar elevation
+		double SolarElevationcorrectedforatmrefractionDeg = SolarElevationAngleDeg + ApproxAtmosphericRefractionDeg;
+
+		// Get the solar azimuth 
+		double tmp = FMath::RadiansToDegrees(FMath::Acos(((FMath::Sin(LatitudeRad)*FMath::Cos(SolarZenithAngleRad)) - FMath::Sin(SunDeclinRad)) / (FMath::Cos(LatitudeRad)*FMath::Sin(SolarZenithAngleRad))));
+		double SolarAzimuthAngleDegcwfromN = HourAngleDeg > 0.0 ? FMath::Fmod(tmp + 180.0, 360.0) : FMath::Fmod(540.0 - tmp, 360.0);
+
+
+		// offset elevation angle to fit with UE coords system
+		SunPositionData.Elevation = 180.0 + SolarElevationAngleDeg;
+		SunPositionData.CorrectedElevation = 180.0 + SolarElevationcorrectedforatmrefractionDeg;
+		SunPositionData.Azimuth = SolarAzimuthAngleDegcwfromN;
+		SunPositionData.SolarNoon = FTimespan::FromDays(SolarNoonLST);
+		SunPositionData.SunriseTime = FTimespan::FromDays(SunriseTimeLST);
+		SunPositionData.SunsetTime = FTimespan::FromDays(SunsetTimeLST);
+	}
+}
+
+UDaySequence* FSunPositionSequence::BuildSequence(UProceduralDaySequenceBuilder* InBuilder)
+{
+	using namespace UE::DaySequence;
+
+	ADaySequenceActor* TargetActor = WeakTargetActor.Get();
+	
+	if (!TargetActor)
+	{
+		return nullptr;
+	}
+
+	if (UDirectionalLightComponent* SunComponent = GetComponentByName<UDirectionalLightComponent>(TargetActor, SunComponentName))
+	{
+		InBuilder->SetActiveBoundObject(SunComponent);
+		
+		const double NormalizedTimeIncrement = 1.0 / FMath::Max(KeyCount - 1, static_cast<unsigned>(1));
+		for (unsigned int Key = 0; Key < KeyCount; ++Key)
+		{
+			const double KeyTime = Key * NormalizedTimeIncrement;
+		
+			FSunPositionData Data;
+
+			double Hours = (double)Key / (double)KeyCount * 24.0;
+			int32 IntHours = Hours;
+
+			double Minutes = FMath::Frac(Hours) * 60.0;
+			int32 IntMinutes = Minutes;
+
+			double Seconds = FMath::Frac(Minutes) * 60.0;
+			int32 IntSeconds = Seconds;
+
+			double Milliseconds = FMath::Frac(Seconds) * 1000.0;
+			int32 IntMilliseconds = Milliseconds;
+		
+			GetSunPosition(Latitude, Longitude, TimeZone, bIsDaylightSavings, Time.GetYear(), Time.GetMonth(), Time.GetDay(), IntHours, IntMinutes, IntSeconds, IntMilliseconds, Data);
+
+			InBuilder->AddRotationKey(KeyTime, FRotator(Data.CorrectedElevation, Data.Azimuth, 0.0));
+		}
+	}
+	
+	return ProceduralSequence;
+}
\ No newline at end of file
diff --git a/Engine/Plugins/Experimental/DaySequence/Source/DaySequence/Public/Actors/ProceduralDaySequenceActor.h b/Engine/Plugins/Experimental/DaySequence/Source/DaySequence/Public/Actors/ProceduralDaySequenceActor.h
deleted file mode 100644
diff --git a/Engine/Plugins/Experimental/DaySequence/Source/DaySequence/Public/Actors/SimpleDaySequenceActor.h b/Engine/Plugins/Experimental/DaySequence/Source/DaySequence/Public/Actors/SimpleDaySequenceActor.h
--- a/Engine/Plugins/Experimental/DaySequence/Source/DaySequence/Public/Actors/SimpleDaySequenceActor.h
+++ b/Engine/Plugins/Experimental/DaySequence/Source/DaySequence/Public/Actors/SimpleDaySequenceActor.h
@@ -16,7 +16,7 @@ class UStaticMeshComponent;
 
 /**
  * A self-registering Day Sequence Actor with a simple set of lighting components (some of which are optional).
- * Can be used as-is or extended by derived classes (see AProceduralDaySequenceActor for example).
+ * Can be used as-is or extended by derived classes (see ASunPositionDaySequenceActor).
  */
 UCLASS(Blueprintable, HideCategories=(Tags, Networking, LevelInstance))
 class DAYSEQUENCE_API ASimpleDaySequenceActor
@@ -28,11 +28,13 @@ public:
 	ASimpleDaySequenceActor(const FObjectInitializer& Init);
 
 protected:
+	
 	/** BeginPlay and OnConstruction overrides auto-register this actor with the DaySequenceSubsystem. */
 	virtual void BeginPlay() override;
 	virtual void OnConstruction(const FTransform& Transform) override;
 
 protected:
+
 	/** Standard Components **/
 	
 	UPROPERTY(VisibleDefaultsOnly, BlueprintReadOnly, Category= "Day Sequence", meta = (AllowPrivateAccess = "true"))
diff --git a/Engine/Plugins/Experimental/DaySequence/Source/DaySequence/Public/Actors/SunAngleDaySequenceActor.h b/Engine/Plugins/Experimental/DaySequence/Source/DaySequence/Public/Actors/SunAngleDaySequenceActor.h
deleted file mode 100644
diff --git a/Engine/Plugins/Experimental/DaySequence/Source/DaySequence/Public/Actors/SunPositionDaySequenceActor.h b/Engine/Plugins/Experimental/DaySequence/Source/DaySequence/Public/Actors/SunPositionDaySequenceActor.h
--- a/Engine/Plugins/Experimental/DaySequence/Source/DaySequence/Public/Actors/SunPositionDaySequenceActor.h
+++ b/Engine/Plugins/Experimental/DaySequence/Source/DaySequence/Public/Actors/SunPositionDaySequenceActor.h
@@ -3,72 +3,23 @@
 #pragma once
 
 #include "CoreMinimal.h"
-#include "SunAngleDaySequenceActor.h"
+#include "SimpleDaySequenceActor.h"
 
 #include "SunPositionDaySequenceActor.generated.h"
 
 /**
- * A simple procedural Day Sequence Actor that mimics an Earth day at a user-specified location and time.
+ * A Day Sequence Actor that represents a physically accurate 24 hour day cycle.
  */
 UCLASS(Blueprintable)
 class DAYSEQUENCE_API ASunPositionDaySequenceActor
-	: public AProceduralDaySequenceActor
+	: public ASimpleDaySequenceActor
 {
 	GENERATED_BODY()
 
 public:
 	ASunPositionDaySequenceActor(const FObjectInitializer& Init);
 
-protected:
-	virtual void Tick(float DeltaTime) override;
-
-	void ApplySunAngle();
-	
-	virtual void BuildSequence(UProceduralDaySequenceBuilder* SequenceBuilder) override;
-	virtual void SequencePlayerUpdated(float CurrentTime, float PreviousTime) override;
-
-#if WITH_EDITOR
-	virtual void PostEditChangeProperty(FPropertyChangedEvent& PropertyChangedEvent) override;
-#endif
-
 protected:
 	UPROPERTY(VisibleDefaultsOnly, BlueprintReadOnly, Category= "Day Sequence", meta = (AllowPrivateAccess = "true"))
 	TObjectPtr<UDirectionalLightComponent> MoonComponent;
-	
-private:
-	
-	/** User settings */
-	UPROPERTY(EditAnywhere, Category = "Procedural Parameters")
-	unsigned KeyCount = 24;
-	
-	UPROPERTY(EditAnywhere, Category = "Procedural Parameters")
-	FDateTime Time;
-
-	UPROPERTY(EditAnywhere, Category = "Procedural Parameters")
-	double TimeZone;
-	
-	UPROPERTY(EditAnywhere, Category = "Procedural Parameters")
-	double Latitude;
-
-	UPROPERTY(EditAnywhere, Category = "Procedural Parameters")
-	double Longitude;
-	
-	UPROPERTY(EditAnywhere, Category = "Procedural Parameters")
-	bool bIsDaylightSavings;
-
-	
-	/** Read only properties */
-	UPROPERTY(VisibleAnywhere, Transient, Category=Sun, meta=(NoResetToDefault))
-	double Elevation = 0.0f;
-	
-	UPROPERTY(VisibleAnywhere, Transient, Category=Sun, meta=(NoResetToDefault))
-	double Azimuth = 0.0f;
-
-	
-	/** Animated properties */
-	UPROPERTY(Transient, Interp, Category=Sun)
-	double SunPitch = 0.0f;
-
-	UPROPERTY(Transient, Interp, Category=Sun)
-	double SunYaw = 0.f;
 };
\ No newline at end of file
diff --git a/Engine/Plugins/Experimental/DaySequence/Source/DaySequence/Public/DaySequenceCollectionAsset.h b/Engine/Plugins/Experimental/DaySequence/Source/DaySequence/Public/DaySequenceCollectionAsset.h
--- a/Engine/Plugins/Experimental/DaySequence/Source/DaySequence/Public/DaySequenceCollectionAsset.h
+++ b/Engine/Plugins/Experimental/DaySequence/Source/DaySequence/Public/DaySequenceCollectionAsset.h
@@ -4,6 +4,8 @@
 
 #include "DaySequenceConditionSet.h"
 #include "Engine/DataAsset.h"
+#include "ProceduralDaySequence.h"
+#include "StructUtils/InstancedStruct.h"
 
 #include "DaySequenceCollectionAsset.generated.h"
 
@@ -46,6 +48,9 @@ public:
 	UPROPERTY(EditAnywhere, Category="Day Sequence")
 	TArray<FDaySequenceCollectionEntry> DaySequences;
 
+	UPROPERTY(EditAnywhere, Category="Day Sequence", meta = (ExcludeBaseStruct))
+	TArray<TInstancedStruct<FProceduralDaySequence>> ProceduralDaySequences;
+
 	/**
 	 * TODO [nickolas.drake]
 	 * Add an API which:
diff --git a/Engine/Plugins/Experimental/DaySequence/Source/DaySequence/Public/ProceduralDaySequence.h b/Engine/Plugins/Experimental/DaySequence/Source/DaySequence/Public/ProceduralDaySequence.h
new file mode 100644
--- /dev/null
+++ b/Engine/Plugins/Experimental/DaySequence/Source/DaySequence/Public/ProceduralDaySequence.h
@@ -0,0 +1,52 @@
+// Copyright Epic Games, Inc. All Rights Reserved.
+
+#pragma once
+
+#include "ProceduralDaySequenceBuilder.h"
+
+#include "ProceduralDaySequence.generated.h"
+
+class ADaySequenceActor;
+class UDaySequence;
+
+namespace UE::DaySequence
+{
+	// Utility function to simplify looking for owned components by type and name.
+	template <typename T>
+	T* GetComponentByName(AActor* InActor, FName Name)
+	{
+		for (T* Component : TInlineComponentArray<T*>(InActor))
+		{
+			if (Component->GetFName() == Name)
+			{
+				return Component;
+			}
+		}
+		
+		return nullptr;
+	}
+}
+
+/**
+ * Base class for procedural sequences.
+ * To create a procedural sequence, a subclass of this type should be created that overrides BuildSequence.
+ * See FSunPositionSequence, FSunAngleSequence, and FSineSequence for examples.
+ */
+USTRUCT(meta=(Hidden))
+struct DAYSEQUENCE_API FProceduralDaySequence
+{
+	GENERATED_BODY()
+	
+	virtual ~FProceduralDaySequence()
+	{}
+	
+	UDaySequence* GetSequence(ADaySequenceActor* InActor);
+
+protected:
+	virtual UDaySequence* BuildSequence(UProceduralDaySequenceBuilder* InBuilder) { return nullptr; }
+
+	UPROPERTY(Transient)
+	TObjectPtr<UDaySequence> ProceduralSequence = nullptr;
+
+	TWeakObjectPtr<ADaySequenceActor> WeakTargetActor = nullptr;
+};
\ No newline at end of file
diff --git a/Engine/Plugins/Experimental/DaySequence/Source/DaySequence/Public/ProceduralDaySequenceBuilder.h b/Engine/Plugins/Experimental/DaySequence/Source/DaySequence/Public/ProceduralDaySequenceBuilder.h
--- a/Engine/Plugins/Experimental/DaySequence/Source/DaySequence/Public/ProceduralDaySequenceBuilder.h
+++ b/Engine/Plugins/Experimental/DaySequence/Source/DaySequence/Public/ProceduralDaySequenceBuilder.h
@@ -4,16 +4,19 @@
 
 #include "CoreMinimal.h"
 #include "DaySequence.h"
-#include "DaySequenceActor.h"
 #include "Curves/RealCurve.h"
 
 #include "ProceduralDaySequenceBuilder.generated.h"
 
+class ADaySequenceActor;
+
 /**
  * A utility class for creating procedural Day Sequences.
  * Before adding any keys, SetActiveBoundObject should be called and provided a Day Sequence Actor or a component owned by a Day Sequence Actor.
  * All time values are currently normalized to the range [0, 1], inclusive on both ends. A time of 1 is handled as a special case and maps to the final frame.
  * This class assumes the target Day Sequence Actor will stay alive and that users will keep the generated sequence alive, it manages no lifetimes.
+ *
+ * Consider using FProceduralDaySequence instead of using this class directly.
  */
 UCLASS(BlueprintType)
 class DAYSEQUENCE_API UProceduralDaySequenceBuilder
@@ -41,6 +44,8 @@ public:
 	UFUNCTION(BlueprintCallable, Category = "Day Sequence")
 	void SetActiveBoundObject(UObject* InObject);
 
+	// todo [nickolas.drake]: Fill out this section, making sure common key types can be added in bulk and by BPs
+	
 	UFUNCTION(BlueprintCallable, Category = "Day Sequence")
 	void ClearAllKeys();
 	
@@ -57,6 +62,18 @@ public:
 	void AddBoolKey(FName PropertyName, const TPair<float, bool>& KeyValue);
 	void AddBoolKeys(FName PropertyName, const TArray<TPair<float, bool>>& KeysAndValues);
 
+	UFUNCTION(BlueprintCallable, Category = "Day Sequence")
+	void AddTransformKey(float Key, const FTransform& Value, ERichCurveInterpMode InterpMode = RCIM_Cubic);
+	
+	UFUNCTION(BlueprintCallable, Category = "Day Sequence")
+	void AddTranslationKey(float Key, const FVector& Value, ERichCurveInterpMode InterpMode = RCIM_Cubic);
+	
+	UFUNCTION(BlueprintCallable, Category = "Day Sequence")
+	void AddRotationKey(float Key, const FRotator& Value, ERichCurveInterpMode InterpMode = RCIM_Cubic);
+	
+	UFUNCTION(BlueprintCallable, Category = "Day Sequence")
+	void AddScaleKey(float Key, const FVector& Value, ERichCurveInterpMode InterpMode = RCIM_Cubic);
+
 	UFUNCTION(BlueprintCallable, Category = "Day Sequence")
 	void AddStaticTime(float StaticTime);
 	
diff --git a/Engine/Plugins/Experimental/DaySequence/Source/DaySequence/Public/ProceduralSequences/SineSequence.h b/Engine/Plugins/Experimental/DaySequence/Source/DaySequence/Public/ProceduralSequences/SineSequence.h
new file mode 100644
--- /dev/null
+++ b/Engine/Plugins/Experimental/DaySequence/Source/DaySequence/Public/ProceduralSequences/SineSequence.h
@@ -0,0 +1,46 @@
+// Copyright Epic Games, Inc. All Rights Reserved.
+
+#pragma once
+
+#include "ProceduralDaySequence.h"
+
+#include "SineSequence.generated.h"
+
+/**
+ * A procedural sequence that animates a user specified property according to a sine wave.
+ */
+USTRUCT()
+struct DAYSEQUENCE_API FSineSequence : public FProceduralDaySequence
+{
+	GENERATED_BODY()
+
+	virtual ~FSineSequence() override
+	{}
+	
+	// The name of the float property to animate.
+	UPROPERTY(EditAnywhere, Category = "Procedural Parameters")
+	FName PropertyName = FName(TEXT(""));
+	
+	// Optional component, used if the property being animated is on a component owned by the Day Sequence Actor.
+	UPROPERTY(EditAnywhere, Category = "Procedural Parameters")
+	FName ComponentName = FName(TEXT(""));
+
+	UPROPERTY(EditAnywhere, Category = "Procedural Parameters")
+	unsigned KeyCount = 5;
+
+	UPROPERTY(EditAnywhere, Category = "Procedural Parameters")
+	float Amplitude = 1.f;
+	
+	UPROPERTY(EditAnywhere, Category = "Procedural Parameters")
+	float Frequency = 1.f;
+	
+	UPROPERTY(EditAnywhere, Category = "Procedural Parameters")
+	float PhaseShift = 0.f;
+
+	UPROPERTY(EditAnywhere, Category = "Procedural Parameters")
+	float VerticalShift = 0.f;
+	
+private:
+	
+	virtual UDaySequence* BuildSequence(UProceduralDaySequenceBuilder* InBuilder) override;
+};
\ No newline at end of file
diff --git a/Engine/Plugins/Experimental/DaySequence/Source/DaySequence/Public/ProceduralSequences/SunAngleSequence.h b/Engine/Plugins/Experimental/DaySequence/Source/DaySequence/Public/ProceduralSequences/SunAngleSequence.h
new file mode 100644
--- /dev/null
+++ b/Engine/Plugins/Experimental/DaySequence/Source/DaySequence/Public/ProceduralSequences/SunAngleSequence.h
@@ -0,0 +1,26 @@
+// Copyright Epic Games, Inc. All Rights Reserved.
+
+#pragma once
+
+#include "ProceduralDaySequence.h"
+
+#include "SunAngleSequence.generated.h"
+
+/**
+ * A procedural sequence that linearly animates the sun.
+ */
+USTRUCT()
+struct DAYSEQUENCE_API FSunAngleSequence : public FProceduralDaySequence
+{
+	GENERATED_BODY()
+
+	virtual ~FSunAngleSequence() override
+	{}
+
+	UPROPERTY(EditAnywhere, Category = "Procedural Parameters")
+	FName SunComponentName = FName(TEXT("Sun"));
+
+private:
+	
+	virtual UDaySequence* BuildSequence(UProceduralDaySequenceBuilder* InBuilder) override;
+};
\ No newline at end of file
diff --git a/Engine/Plugins/Experimental/DaySequence/Source/DaySequence/Public/ProceduralSequences/SunPositionSequence.h b/Engine/Plugins/Experimental/DaySequence/Source/DaySequence/Public/ProceduralSequences/SunPositionSequence.h
new file mode 100644
--- /dev/null
+++ b/Engine/Plugins/Experimental/DaySequence/Source/DaySequence/Public/ProceduralSequences/SunPositionSequence.h
@@ -0,0 +1,44 @@
+// Copyright Epic Games, Inc. All Rights Reserved.
+
+#pragma once
+
+#include "ProceduralDaySequence.h"
+
+#include "SunPositionSequence.generated.h"
+
+/**
+ * A procedural sequence that animates a sun in a physically accurate way based on geographic data.
+ */
+USTRUCT()
+struct DAYSEQUENCE_API FSunPositionSequence : public FProceduralDaySequence
+{
+	GENERATED_BODY()
+
+	virtual ~FSunPositionSequence() override
+	{}
+
+	UPROPERTY(EditAnywhere, Category = "Procedural Parameters")
+	FName SunComponentName = FName(TEXT("Sun"));
+	
+	UPROPERTY(EditAnywhere, Category = "Procedural Parameters")
+	unsigned KeyCount = 24;
+	
+	UPROPERTY(EditAnywhere, Category = "Procedural Parameters")
+	FDateTime Time;
+
+	UPROPERTY(EditAnywhere, Category = "Procedural Parameters")
+	double TimeZone;
+	
+	UPROPERTY(EditAnywhere, Category = "Procedural Parameters")
+	double Latitude;
+
+	UPROPERTY(EditAnywhere, Category = "Procedural Parameters")
+	double Longitude;
+	
+	UPROPERTY(EditAnywhere, Category = "Procedural Parameters")
+	bool bIsDaylightSavings;
+	
+private:
+	
+	virtual UDaySequence* BuildSequence(UProceduralDaySequenceBuilder* InBuilder) override;
+};
\ No newline at end of file
