DaySequenceActor: Include proper header to fix non-unity issue.

diff --git a/Engine/Plugins/Experimental/DaySequence/Source/DaySequence/Public/ProceduralDaySequence.h b/Engine/Plugins/Experimental/DaySequence/Source/DaySequence/Public/ProceduralDaySequence.h
--- a/Engine/Plugins/Experimental/DaySequence/Source/DaySequence/Public/ProceduralDaySequence.h
+++ b/Engine/Plugins/Experimental/DaySequence/Source/DaySequence/Public/ProceduralDaySequence.h
@@ -4,6 +4,8 @@
 
 #include "ProceduralDaySequenceBuilder.h"
 
+#include "GameFramework/Actor.h"
+
 #include "ProceduralDaySequence.generated.h"
 
 class ADaySequenceActor;
