ProceduralDaySequence: Fix GC prevention.

diff --git a/Engine/Plugins/Experimental/DaySequence/Source/DaySequence/Private/ProceduralDaySequence.cpp b/Engine/Plugins/Experimental/DaySequence/Source/DaySequence/Private/ProceduralDaySequence.cpp
--- a/Engine/Plugins/Experimental/DaySequence/Source/DaySequence/Private/ProceduralDaySequence.cpp
+++ b/Engine/Plugins/Experimental/DaySequence/Source/DaySequence/Private/ProceduralDaySequence.cpp
@@ -12,12 +12,13 @@
 UDaySequence* FProceduralDaySequence::GetSequence(ADaySequenceActor* InActor)
 {
 	WeakTargetActor = InActor;
-	ProceduralSequence = nullptr;
 	
-	if (!ProceduralSequence && InActor)
+	UDaySequence* ProceduralSequence = nullptr;
+	
+	if (InActor)
 	{
 		UProceduralDaySequenceBuilder* SequenceBuilder = NewObject<UProceduralDaySequenceBuilder>();
-		ProceduralSequence = SequenceBuilder->Initialize(InActor, ProceduralSequence);
+		ProceduralSequence = SequenceBuilder->Initialize(InActor);
 		BuildSequence(SequenceBuilder);
 	}
 	
diff --git a/Engine/Plugins/Experimental/DaySequence/Source/DaySequence/Private/ProceduralSequences/SineSequence.cpp b/Engine/Plugins/Experimental/DaySequence/Source/DaySequence/Private/ProceduralSequences/SineSequence.cpp
--- a/Engine/Plugins/Experimental/DaySequence/Source/DaySequence/Private/ProceduralSequences/SineSequence.cpp
+++ b/Engine/Plugins/Experimental/DaySequence/Source/DaySequence/Private/ProceduralSequences/SineSequence.cpp
@@ -7,7 +7,7 @@
 
 #include UE_INLINE_GENERATED_CPP_BY_NAME(SineSequence)
 
-UDaySequence* FSineSequence::BuildSequence(UProceduralDaySequenceBuilder* InBuilder)
+void FSineSequence::BuildSequence(UProceduralDaySequenceBuilder* InBuilder)
 {
 	using namespace UE::DaySequence;
 
@@ -15,7 +15,7 @@ UDaySequence* FSineSequence::BuildSequence(UProceduralDaySequenceBuilder* InBuil
 	
 	if (!TargetActor || PropertyName.IsNone())
 	{
-		return nullptr;
+		return;
 	}
 
 	UObject* AnimatedObject = TargetActor;
@@ -37,6 +37,4 @@ UDaySequence* FSineSequence::BuildSequence(UProceduralDaySequenceBuilder* InBuil
 			InBuilder->AddScalarKey(PropertyName, KeyTime, Amplitude * FMath::Sin(UE_DOUBLE_TWO_PI * Frequency * (KeyTime - PhaseShift)) + VerticalShift);
 		}
 	}
-	
-	return ProceduralSequence;
 }
\ No newline at end of file
diff --git a/Engine/Plugins/Experimental/DaySequence/Source/DaySequence/Private/ProceduralSequences/SunAngleSequence.cpp b/Engine/Plugins/Experimental/DaySequence/Source/DaySequence/Private/ProceduralSequences/SunAngleSequence.cpp
--- a/Engine/Plugins/Experimental/DaySequence/Source/DaySequence/Private/ProceduralSequences/SunAngleSequence.cpp
+++ b/Engine/Plugins/Experimental/DaySequence/Source/DaySequence/Private/ProceduralSequences/SunAngleSequence.cpp
@@ -9,7 +9,7 @@
 
 #include UE_INLINE_GENERATED_CPP_BY_NAME(SunAngleSequence)
 
-UDaySequence* FSunAngleSequence::BuildSequence(UProceduralDaySequenceBuilder* InBuilder)
+void FSunAngleSequence::BuildSequence(UProceduralDaySequenceBuilder* InBuilder)
 {
 	using namespace UE::DaySequence;
 
@@ -17,7 +17,7 @@ UDaySequence* FSunAngleSequence::BuildSequence(UProceduralDaySequenceBuilder* In
 	
 	if (!TargetActor)
 	{
-		return nullptr;
+		return;
 	}
 	
 	if (UDirectionalLightComponent* SunComponent = GetComponentByName<UDirectionalLightComponent>(TargetActor, SunComponentName))
@@ -27,6 +27,4 @@ UDaySequence* FSunAngleSequence::BuildSequence(UProceduralDaySequenceBuilder* In
 		InBuilder->AddRotationKey(0.0, FRotator(90.0, 0.0, 0.0), RCIM_Linear);
 		InBuilder->AddRotationKey(1.0, FRotator(450.0, 0.0, 0.0), RCIM_Linear);
 	}
-	
-	return ProceduralSequence;
 }
\ No newline at end of file
diff --git a/Engine/Plugins/Experimental/DaySequence/Source/DaySequence/Private/ProceduralSequences/SunPositionSequence.cpp b/Engine/Plugins/Experimental/DaySequence/Source/DaySequence/Private/ProceduralSequences/SunPositionSequence.cpp
--- a/Engine/Plugins/Experimental/DaySequence/Source/DaySequence/Private/ProceduralSequences/SunPositionSequence.cpp
+++ b/Engine/Plugins/Experimental/DaySequence/Source/DaySequence/Private/ProceduralSequences/SunPositionSequence.cpp
@@ -168,7 +168,7 @@ namespace UE::DaySequence
 	}
 }
 
-UDaySequence* FSunPositionSequence::BuildSequence(UProceduralDaySequenceBuilder* InBuilder)
+void FSunPositionSequence::BuildSequence(UProceduralDaySequenceBuilder* InBuilder)
 {
 	using namespace UE::DaySequence;
 
@@ -176,7 +176,7 @@ UDaySequence* FSunPositionSequence::BuildSequence(UProceduralDaySequenceBuilder*
 	
 	if (!TargetActor)
 	{
-		return nullptr;
+		return;
 	}
 
 	if (UDirectionalLightComponent* SunComponent = GetComponentByName<UDirectionalLightComponent>(TargetActor, SunComponentName))
@@ -207,6 +207,4 @@ UDaySequence* FSunPositionSequence::BuildSequence(UProceduralDaySequenceBuilder*
 			InBuilder->AddRotationKey(KeyTime, FRotator(Data.CorrectedElevation, Data.Azimuth, 0.0));
 		}
 	}
-	
-	return ProceduralSequence;
 }
\ No newline at end of file
diff --git a/Engine/Plugins/Experimental/DaySequence/Source/DaySequence/Public/ProceduralDaySequence.h b/Engine/Plugins/Experimental/DaySequence/Source/DaySequence/Public/ProceduralDaySequence.h
--- a/Engine/Plugins/Experimental/DaySequence/Source/DaySequence/Public/ProceduralDaySequence.h
+++ b/Engine/Plugins/Experimental/DaySequence/Source/DaySequence/Public/ProceduralDaySequence.h
@@ -45,10 +45,8 @@ struct DAYSEQUENCE_API FProceduralDaySequence
 	UDaySequence* GetSequence(ADaySequenceActor* InActor);
 
 protected:
-	virtual UDaySequence* BuildSequence(UProceduralDaySequenceBuilder* InBuilder) { return nullptr; }
-
-	UPROPERTY(Transient)
-	TObjectPtr<UDaySequence> ProceduralSequence = nullptr;
-
+	
+	virtual void BuildSequence(UProceduralDaySequenceBuilder* InBuilder) {}
+	
 	TWeakObjectPtr<ADaySequenceActor> WeakTargetActor = nullptr;
 };
\ No newline at end of file
diff --git a/Engine/Plugins/Experimental/DaySequence/Source/DaySequence/Public/ProceduralSequences/SineSequence.h b/Engine/Plugins/Experimental/DaySequence/Source/DaySequence/Public/ProceduralSequences/SineSequence.h
--- a/Engine/Plugins/Experimental/DaySequence/Source/DaySequence/Public/ProceduralSequences/SineSequence.h
+++ b/Engine/Plugins/Experimental/DaySequence/Source/DaySequence/Public/ProceduralSequences/SineSequence.h
@@ -42,5 +42,5 @@ struct DAYSEQUENCE_API FSineSequence : public FProceduralDaySequence
 	
 private:
 	
-	virtual UDaySequence* BuildSequence(UProceduralDaySequenceBuilder* InBuilder) override;
+	virtual void BuildSequence(UProceduralDaySequenceBuilder* InBuilder) override;
 };
\ No newline at end of file
diff --git a/Engine/Plugins/Experimental/DaySequence/Source/DaySequence/Public/ProceduralSequences/SunAngleSequence.h b/Engine/Plugins/Experimental/DaySequence/Source/DaySequence/Public/ProceduralSequences/SunAngleSequence.h
--- a/Engine/Plugins/Experimental/DaySequence/Source/DaySequence/Public/ProceduralSequences/SunAngleSequence.h
+++ b/Engine/Plugins/Experimental/DaySequence/Source/DaySequence/Public/ProceduralSequences/SunAngleSequence.h
@@ -22,5 +22,5 @@ struct DAYSEQUENCE_API FSunAngleSequence : public FProceduralDaySequence
 
 private:
 	
-	virtual UDaySequence* BuildSequence(UProceduralDaySequenceBuilder* InBuilder) override;
+	virtual void BuildSequence(UProceduralDaySequenceBuilder* InBuilder) override;
 };
\ No newline at end of file
diff --git a/Engine/Plugins/Experimental/DaySequence/Source/DaySequence/Public/ProceduralSequences/SunPositionSequence.h b/Engine/Plugins/Experimental/DaySequence/Source/DaySequence/Public/ProceduralSequences/SunPositionSequence.h
--- a/Engine/Plugins/Experimental/DaySequence/Source/DaySequence/Public/ProceduralSequences/SunPositionSequence.h
+++ b/Engine/Plugins/Experimental/DaySequence/Source/DaySequence/Public/ProceduralSequences/SunPositionSequence.h
@@ -40,5 +40,5 @@ struct DAYSEQUENCE_API FSunPositionSequence : public FProceduralDaySequence
 	
 private:
 	
-	virtual UDaySequence* BuildSequence(UProceduralDaySequenceBuilder* InBuilder) override;
+	virtual void BuildSequence(UProceduralDaySequenceBuilder* InBuilder) override;
 };
\ No newline at end of file
