DaySequence: Several renames

SimpleDaySequenceActor -> BaseDaySequenceActor
SunPositionDaySequenceActor -> SunMoonDaySequenceActor

diff --git a/Engine/Plugins/Experimental/DaySequence/Source/DaySequence/Private/Actors/BaseDaySequenceActor.cpp b/Engine/Plugins/Experimental/DaySequence/Source/DaySequence/Private/Actors/BaseDaySequenceActor.cpp
new file mode 100644
--- /dev/null
+++ b/Engine/Plugins/Experimental/DaySequence/Source/DaySequence/Private/Actors/BaseDaySequenceActor.cpp
@@ -0,0 +1,79 @@
+﻿// Copyright Epic Games, Inc. All Rights Reserved.
+
+#include "Actors/BaseDaySequenceActor.h"
+
+#include "DaySequenceSubsystem.h"
+
+#include "Components/DirectionalLightComponent.h"
+#include "Components/ExponentialHeightFogComponent.h"
+#include "Components/SkyAtmosphereComponent.h"
+#include "Components/SkyLightComponent.h"
+#include "Components/VolumetricCloudComponent.h"
+#include "Components/StaticMeshComponent.h"
+
+#include "Engine/World.h"
+#include "Materials/MaterialInterface.h"
+#include "Engine/StaticMesh.h"
+#include "UObject/ConstructorHelpers.h"
+
+ABaseDaySequenceActor::ABaseDaySequenceActor(const FObjectInitializer& Init)
+: Super(Init)
+{
+	SunRootComponent = CreateDefaultSubobject<USceneComponent>(TEXT("SunRoot"));
+	SunRootComponent->SetupAttachment(RootComponent);
+	
+	SunComponent = CreateDefaultSubobject<UDirectionalLightComponent>(TEXT("Sun"));
+	SunComponent->SetupAttachment(SunRootComponent);
+	
+	ExponentialHeightFogComponent = CreateOptionalDefaultSubobject<UExponentialHeightFogComponent>(TEXT("ExponentialHeightFog"));
+	ExponentialHeightFogComponent->SetupAttachment(RootComponent);
+	ExponentialHeightFogComponent->bEnableVolumetricFog = true;
+
+	SkyAtmosphereComponent = CreateDefaultSubobject<USkyAtmosphereComponent>(TEXT("SkyAtmosphere"));
+	SkyAtmosphereComponent->SetupAttachment(RootComponent);
+
+	SkyLightComponent = CreateDefaultSubobject<USkyLightComponent>(TEXT("SkyLight"));
+	SkyLightComponent->SetupAttachment(RootComponent);
+	SkyLightComponent->bRealTimeCapture = true;
+	SkyLightComponent->bLowerHemisphereIsBlack = false;
+
+	VolumetricCloudComponent = CreateOptionalDefaultSubobject<UVolumetricCloudComponent>(TEXT("VolumetricCloud"));
+    VolumetricCloudComponent->SetupAttachment(RootComponent);
+
+	static ConstructorHelpers::FObjectFinder<UStaticMesh> SkySphereDefaultMesh(TEXT("/Engine/EngineSky/SM_SkySphere.SM_SkySphere"));
+	static ConstructorHelpers::FObjectFinder<UMaterialInterface> SkySphereDefaultMaterial(TEXT("/Engine/EngineSky/M_SimpleSkyDome.M_SimpleSkyDome"));
+	SkySphereComponent = CreateOptionalDefaultSubobject<UStaticMeshComponent>(TEXT("SkySphere"));
+	SkySphereComponent->SetupAttachment(RootComponent);
+	SkySphereComponent->SetStaticMesh(SkySphereDefaultMesh.Object);
+	SkySphereComponent->SetMaterial(0, SkySphereDefaultMaterial.Object.Get());
+	SkySphereComponent->SetRelativeScale3D(FVector(400.f));
+}
+
+void ABaseDaySequenceActor::BeginPlay()
+{
+	Super::BeginPlay();
+	
+	if (const UWorld* World = GetWorld())
+	{
+		if (UDaySequenceSubsystem* DaySequenceSubsystem = World->GetSubsystem<UDaySequenceSubsystem>())
+		{
+			DaySequenceSubsystem->SetDaySequenceActor(this);
+		}
+	}
+}
+
+void ABaseDaySequenceActor::OnConstruction(const FTransform& Transform)
+{
+	Super::OnConstruction(Transform);
+	
+	if (const UWorld* World = GetWorld())
+	{
+		if (UDaySequenceSubsystem* DaySequenceSubsystem = World->GetSubsystem<UDaySequenceSubsystem>())
+		{
+			if (DaySequenceSubsystem->GetDaySequenceActor(/* bFindFallbackOnNull */ false) != this)
+			{
+				DaySequenceSubsystem->SetDaySequenceActor(this);
+			}
+		}
+	}
+}
\ No newline at end of file
diff --git a/Engine/Plugins/Experimental/DaySequence/Source/DaySequence/Private/Actors/SimpleDaySequenceActor.cpp b/Engine/Plugins/Experimental/DaySequence/Source/DaySequence/Private/Actors/SimpleDaySequenceActor.cpp
deleted file mode 100644
diff --git a/Engine/Plugins/Experimental/DaySequence/Source/DaySequence/Private/Actors/SunMoonDaySequenceActor.cpp b/Engine/Plugins/Experimental/DaySequence/Source/DaySequence/Private/Actors/SunMoonDaySequenceActor.cpp
new file mode 100644
--- /dev/null
+++ b/Engine/Plugins/Experimental/DaySequence/Source/DaySequence/Private/Actors/SunMoonDaySequenceActor.cpp
@@ -0,0 +1,38 @@
+// Copyright Epic Games, Inc. All Rights Reserved.
+
+#include "Actors/SunMoonDaySequenceActor.h"
+
+#include "Components/DirectionalLightComponent.h"
+#include "Components/StaticMeshComponent.h"
+#include "DaySequenceCollectionAsset.h"
+
+#include "Materials/MaterialInterface.h"
+#include "UObject/ConstructorHelpers.h"
+
+ASunMoonDaySequenceActor::ASunMoonDaySequenceActor(const FObjectInitializer& Init)
+: Super(Init)
+{
+	MoonComponent = CreateDefaultSubobject<UDirectionalLightComponent>(TEXT("Moon"));
+	MoonComponent->SetupAttachment(SunRootComponent);
+
+	// Give Sun forward shading priority.
+	SunComponent->SetForwardShadingPriority(1);
+	MoonComponent->SetForwardShadingPriority(0);
+
+	// Configure other Moon defaults.
+	MoonComponent->SetAtmosphereSunLightIndex(1);	// Make Moon the secondary directional light that contributes to the sky atmosphere.
+	MoonComponent->SetIntensity(2.f);
+	MoonComponent->SetUseTemperature(true);
+	MoonComponent->SetTemperature(9000.f);
+	
+	// Override the sky sphere material.
+	static ConstructorHelpers::FObjectFinder<UMaterialInterface> SkySphereDefaultMaterial(TEXT("/DaySequence/MI_24hrSky.MI_24hrSky"));
+	SkySphereComponent->SetMaterial(0, SkySphereDefaultMaterial.Object.Get());
+
+	if (!IsTemplate())
+	{
+		// Override the default collection (which animates the moon and sky material)
+        static ConstructorHelpers::FObjectFinder<UDaySequenceCollectionAsset> DefaultCollection(TEXT("/DaySequence/DSCA_24hr.DSCA_24hr"));
+        DaySequenceCollection = DefaultCollection.Object.Get();
+	}
+}
\ No newline at end of file
diff --git a/Engine/Plugins/Experimental/DaySequence/Source/DaySequence/Private/Actors/SunPositionDaySequenceActor.cpp b/Engine/Plugins/Experimental/DaySequence/Source/DaySequence/Private/Actors/SunPositionDaySequenceActor.cpp
deleted file mode 100644
diff --git a/Engine/Plugins/Experimental/DaySequence/Source/DaySequence/Public/Actors/BaseDaySequenceActor.h b/Engine/Plugins/Experimental/DaySequence/Source/DaySequence/Public/Actors/BaseDaySequenceActor.h
new file mode 100644
--- /dev/null
+++ b/Engine/Plugins/Experimental/DaySequence/Source/DaySequence/Public/Actors/BaseDaySequenceActor.h
@@ -0,0 +1,62 @@
+﻿// Copyright Epic Games, Inc. All Rights Reserved.
+
+#pragma once
+
+#include "CoreMinimal.h"
+#include "DaySequenceActor.h"
+
+#include "BaseDaySequenceActor.generated.h"
+
+class USkyAtmosphereComponent;
+class USkyLightComponent;
+class UVolumetricCloudComponent;
+class UDirectionalLightComponent;
+class UExponentialHeightFogComponent;
+class UStaticMeshComponent;
+
+/**
+ * A self-registering Day Sequence Actor with a simple set of lighting components (some of which are optional).
+ * Can be used as-is or extended by derived classes (see ASunPositionDaySequenceActor).
+ */
+UCLASS(Blueprintable, HideCategories=(Tags, Networking, LevelInstance))
+class DAYSEQUENCE_API ABaseDaySequenceActor
+	: public ADaySequenceActor
+{
+	GENERATED_BODY()
+
+public:
+	ABaseDaySequenceActor(const FObjectInitializer& Init);
+
+protected:
+	
+	/** BeginPlay and OnConstruction overrides auto-register this actor with the DaySequenceSubsystem. */
+	virtual void BeginPlay() override;
+	virtual void OnConstruction(const FTransform& Transform) override;
+
+protected:
+
+	/** Standard Components **/
+	
+	UPROPERTY(VisibleDefaultsOnly, BlueprintReadOnly, Category= "Day Sequence", meta = (AllowPrivateAccess = "true"))
+	TObjectPtr<USceneComponent> SunRootComponent;
+	
+	UPROPERTY(VisibleDefaultsOnly, BlueprintReadOnly, Category= "Day Sequence", meta = (AllowPrivateAccess = "true"))
+	TObjectPtr<UDirectionalLightComponent> SunComponent;
+	
+	UPROPERTY(VisibleDefaultsOnly, BlueprintReadOnly, Category = "Day Sequence", meta = (AllowPrivateAccess = "true"))
+	TObjectPtr<USkyAtmosphereComponent> SkyAtmosphereComponent;
+	
+	UPROPERTY(VisibleDefaultsOnly, BlueprintReadOnly, Category = "Day Sequence", meta = (AllowPrivateAccess = "true"))
+	TObjectPtr<USkyLightComponent> SkyLightComponent;
+	
+	/** Optional Components **/
+	
+	UPROPERTY(VisibleDefaultsOnly, BlueprintReadOnly, Category = "Day Sequence", meta = (AllowPrivateAccess = "true"))
+	TObjectPtr<UExponentialHeightFogComponent> ExponentialHeightFogComponent;
+	
+	UPROPERTY(VisibleDefaultsOnly, BlueprintReadOnly, Category = "Day Sequence", meta = (AllowPrivateAccess = "true"))
+	TObjectPtr<UVolumetricCloudComponent> VolumetricCloudComponent;
+
+	UPROPERTY(VisibleDefaultsOnly, BlueprintReadOnly, Category = "Day Sequence", meta = (AllowPrivateAccess = "true"))
+	TObjectPtr<UStaticMeshComponent> SkySphereComponent;
+};
diff --git a/Engine/Plugins/Experimental/DaySequence/Source/DaySequence/Public/Actors/SimpleDaySequenceActor.h b/Engine/Plugins/Experimental/DaySequence/Source/DaySequence/Public/Actors/SimpleDaySequenceActor.h
deleted file mode 100644
diff --git a/Engine/Plugins/Experimental/DaySequence/Source/DaySequence/Public/Actors/SunMoonDaySequenceActor.h b/Engine/Plugins/Experimental/DaySequence/Source/DaySequence/Public/Actors/SunMoonDaySequenceActor.h
new file mode 100644
--- /dev/null
+++ b/Engine/Plugins/Experimental/DaySequence/Source/DaySequence/Public/Actors/SunMoonDaySequenceActor.h
@@ -0,0 +1,25 @@
+// Copyright Epic Games, Inc. All Rights Reserved.
+
+#pragma once
+
+#include "CoreMinimal.h"
+#include "BaseDaySequenceActor.h"
+
+#include "SunMoonDaySequenceActor.generated.h"
+
+/**
+ * A Day Sequence Actor that represents a physically accurate 24 hour day cycle.
+ */
+UCLASS(Blueprintable)
+class DAYSEQUENCE_API ASunMoonDaySequenceActor
+	: public ABaseDaySequenceActor
+{
+	GENERATED_BODY()
+
+public:
+	ASunMoonDaySequenceActor(const FObjectInitializer& Init);
+
+protected:
+	UPROPERTY(VisibleDefaultsOnly, BlueprintReadOnly, Category= "Day Sequence", meta = (AllowPrivateAccess = "true"))
+	TObjectPtr<UDirectionalLightComponent> MoonComponent;
+};
\ No newline at end of file
diff --git a/Engine/Plugins/Experimental/DaySequence/Source/DaySequence/Public/Actors/SunPositionDaySequenceActor.h b/Engine/Plugins/Experimental/DaySequence/Source/DaySequence/Public/Actors/SunPositionDaySequenceActor.h
deleted file mode 100644
