DaySequence: Rework static time

Our current static time solution has a number of issues
- It cannot be blended
- Modifiers can disable static time at unexpected times (stomping static time set by a different modifier)
- It can't be changed without reenabling a modifier

I am introducing the notion of the StaticTimeManager and StaticTimeContributors. The manager has an API for adding and removing contributors and getting the current static time. Contributors just provide the manager with some callbacks for their desired static time and their weights and the manager handles resolving the effective static time.

#changelist validated
#virtualized

diff --git a/Engine/Plugins/Experimental/DaySequence/Source/DaySequence/Private/DaySequenceActor.cpp b/Engine/Plugins/Experimental/DaySequence/Source/DaySequence/Private/DaySequenceActor.cpp
--- a/Engine/Plugins/Experimental/DaySequence/Source/DaySequence/Private/DaySequenceActor.cpp
+++ b/Engine/Plugins/Experimental/DaySequence/Source/DaySequence/Private/DaySequenceActor.cpp
@@ -8,6 +8,7 @@
 #include "DaySequenceTrack.h"
 #include "DaySequenceConditionTag.h"
 #include "DaySequenceSubsystem.h"
+#include "DaySequenceStaticTime.h"
 
 #include "Components/BillboardComponent.h"
 #include "Curves/CurveFloat.h"
@@ -17,7 +18,6 @@
 #include "MovieSceneTimeHelpers.h"
 #include "MovieSceneBindingOverrides.h"
 #include "Net/UnrealNetwork.h"
-#include "MovieSceneTimeController.h"
 #include "Sections/MovieSceneSubSection.h"
 #include "Tracks/MovieSceneSubTrack.h"
 #include "UObject/ConstructorHelpers.h"
@@ -62,49 +62,6 @@ namespace UE::DaySequence
 		GFrameBudgetMicroseconds,
 		TEXT("(Default: 30us) Approximate max per-frame budget for time-of-day actors in microseconds.")
 	);
-
-
-	struct FStaticTimeControllerOverride : FMovieSceneTimeController
-	{
-		TSharedPtr<FMovieSceneTimeController> PreviousTimeController;
-		float TimeInHours;
-		FFrameTime StaticTime;
-#if WITH_EDITORONLY_DATA
-		FDaySequenceTime PreviousTimeOfDayPreview;
-#endif
-
-		FStaticTimeControllerOverride(FFrameTime InStaticTime, TSharedPtr<FMovieSceneTimeController> InPreviousTimeController)
-			: PreviousTimeController(InPreviousTimeController)
-			, StaticTime(InStaticTime)
-		{}
-
-		void OnTick(float DeltaSeconds, float InPlayRate) override
-		{
-			if (PreviousTimeController)
-			{
-				PreviousTimeController->Tick(DeltaSeconds, InPlayRate);
-			}
-		}
-		void OnStartPlaying(const FQualifiedFrameTime& InStartTime) override
-		{
-			if (PreviousTimeController)
-			{
-				PreviousTimeController->StartPlaying(InStartTime);
-			}
-		}
-		void OnStopPlaying(const FQualifiedFrameTime& InStopTime) override
-		{
-			if (PreviousTimeController)
-			{
-				PreviousTimeController->StopPlaying(InStopTime);
-			}
-		}
-		FFrameTime OnRequestCurrentTime(const FQualifiedFrameTime& InCurrentTime, float InPlayRate)
-		{
-			return StaticTime;
-		}
-	};
-
 } // namespace UE::DaySequence
 
 ADaySequenceActor::ADaySequenceActor(const FObjectInitializer& Init)
@@ -115,6 +72,7 @@ ADaySequenceActor::ADaySequenceActor(const FObjectInitializer& Init)
 , DayLength(24, 0, 0)
 , TimePerCycle(0, 5, 0)
 , InitialTimeOfDay(6, 0, 0)
+, StaticTimeManager(MakeShared<UE::DaySequence::FStaticTimeManager>())
 {
 	USceneComponent* SceneRootComponent = CreateDefaultSubobject<USceneComponent>(USceneComponent::GetDefaultSceneRootVariableName());
 	SetRootComponent(SceneRootComponent);
@@ -515,10 +473,6 @@ void ADaySequenceActor::InitializePlayer()
 
 	if (GetWorld()->IsGameWorld())
 	{
-		// Initializing the sequence player may (will) overwrite the TimeController, so make sure we
-		// maintain the overridden one if SetStaticTimeOfDay was called before Initialize
-		TSharedPtr<FStaticTimeControllerOverride> OldTimeControllerOverride = WeakTimeControllerOverride.Pin();
-
 		if (SequencePlayer)
 		{
 			SequencePlayer->Initialize(RootSequence, this, GetPlaybackSettings(RootSequence));
@@ -526,14 +480,6 @@ void ADaySequenceActor::InitializePlayer()
 
 		if (UDaySequencePlayer* Player = GetSequencePlayerInternal())
 		{
-			if (OldTimeControllerOverride && OldTimeControllerOverride != Player->GetTimeController())
-			{
-				ensure(OldTimeControllerOverride->PreviousTimeController == nullptr);
-
-				OldTimeControllerOverride->PreviousTimeController = Player->GetTimeController();
-				Player->SetTimeControllerDirectly(OldTimeControllerOverride);
-			}
-
 			Player->OnPlay.AddUniqueDynamic(this, &ADaySequenceActor::StopDaySequenceUpdateTimer);
 			Player->OnPause.AddUniqueDynamic(this, &ADaySequenceActor::StartDaySequenceUpdateTimer);
 		}
@@ -1021,36 +967,55 @@ void ADaySequenceActor::WarpEvaluationRange(FMovieSceneEvaluationRange& InOutRan
 
 	const FFrameRate TickRate   = InOutRange.GetFrameRate();
 	const float DayCycleSeconds = TimePerCycle.ToSeconds();
-	const float DayLengthHours  = DayLength.ToSeconds();
+	const float DayLengthHours  = DayLength.ToHours();
 
 	TRange<FFrameTime> Range = InOutRange.GetRange();
 
-	// Warp the lower bound if possible
-	if (Range.GetLowerBound().IsClosed())
+	// Auto bounds checking
+	auto TrySetBounds = [&Range](FFrameTime LowerBound, FFrameTime UpperBound)
 	{
-		float TimeSeconds = static_cast<float>(Range.GetLowerBoundValue() / TickRate);
-		float TimeInHours = DayLengthHours * TimeSeconds / DayCycleSeconds;
+		// Warp the lower bound if possible
+		if (Range.GetLowerBound().IsClosed())
+		{
+			// Set the lower bound value while retaining the inclusivity
+			Range.SetLowerBoundValue(LowerBound);
+		}
 
-		// Remap the time
-		TimeInHours = DayInterpCurve->FloatCurve.Eval(TimeInHours, TimeInHours);
+		// Warp the upper bound if possible
+		if (Range.GetUpperBound().IsClosed())
+		{
+			// Set the upper bound value while retaining the inclusivity
+			Range.SetUpperBoundValue(UpperBound);
+		}
+	};
+
+	// Warp with static time if necessary
+	if (HasStaticTimeOfDay())
+	{
+		const float StaticTimeInGameHours = GetStaticTimeOfDay();
+		
+		const FFrameTime LowerBound = (StaticTimeInGameHours * DayCycleSeconds / DayLengthHours) * TickRate;
 
-		// Set the lower bound value while retaining the inclusivity
-		Range.SetLowerBoundValue((TimeInHours * DayCycleSeconds / DayLengthHours) * TickRate);
+		const FFrameTime UpperBound = (StaticTimeInGameHours * DayCycleSeconds / DayLengthHours) * TickRate;
+
+		TrySetBounds(LowerBound, UpperBound);
 	}
 
-	// Warp the upper bound if possible
-	if (Range.GetUpperBound().IsClosed())
+	// Warp with curve
 	{
-		float TimeSeconds = static_cast<float>(Range.GetUpperBoundValue() / TickRate);
-		float TimeInHours = DayLengthHours * TimeSeconds / DayCycleSeconds;
+		const float LowerBoundTimeSeconds = static_cast<float>(Range.GetLowerBoundValue() / TickRate);
+		float LowerBoundTimeInHours = DayLengthHours * LowerBoundTimeSeconds / DayCycleSeconds;
+		LowerBoundTimeInHours = DayInterpCurve->FloatCurve.Eval(LowerBoundTimeInHours, LowerBoundTimeInHours);
+		const FFrameTime LowerBound = (LowerBoundTimeInHours * DayCycleSeconds / DayLengthHours) * TickRate;
 
-		// Remap the time
-		TimeInHours = DayInterpCurve->FloatCurve.Eval(TimeInHours, TimeInHours);
+		const float UpperBoundTimeSeconds = static_cast<float>(Range.GetUpperBoundValue() / TickRate);
+		float UpperBoundTimeInHours = DayLengthHours * UpperBoundTimeSeconds / DayCycleSeconds;
+		UpperBoundTimeInHours = DayInterpCurve->FloatCurve.Eval(UpperBoundTimeInHours, UpperBoundTimeInHours);
+		const FFrameTime UpperBound = (UpperBoundTimeInHours * DayCycleSeconds / DayLengthHours) * TickRate;
 
-		// Set the upper bound value while retaining the inclusivity
-		Range.SetUpperBoundValue((TimeInHours * DayCycleSeconds / DayLengthHours) * TickRate);
+		TrySetBounds(LowerBound, UpperBound);
 	}
-
+	
 	InOutRange.ResetRange(Range);
 }
 
@@ -1176,7 +1141,11 @@ float ADaySequenceActor::GetTimeOfDay() const
 	}
 	else
 	{
+#if WITH_EDITOR
+		Result = GetTimeOfDayPreview();
+#else
 		Result = GetInitialTimeOfDay();
+#endif
 	}
 	return Result;
 }
@@ -1206,113 +1175,37 @@ bool ADaySequenceActor::SetTimeOfDay(float InHours)
 
 bool ADaySequenceActor::HasStaticTimeOfDay() const
 {
-#if WITH_EDITORONLY_DATA
-	return WeakTimeControllerOverride.Pin() != nullptr || StaticTime.IsSet();
-#else
-	return WeakTimeControllerOverride.Pin() != nullptr;
-#endif
+	return StaticTimeManager && StaticTimeManager->HasStaticTime();
 }
 
 float ADaySequenceActor::GetStaticTimeOfDay() const
 {
-	using namespace UE::DaySequence;
-
-	if (TSharedPtr<FStaticTimeControllerOverride> TimeControllerOverride = WeakTimeControllerOverride.Pin())
+	if (StaticTimeManager && HasStaticTimeOfDay())
 	{
-		return TimeControllerOverride->TimeInHours;
+		return StaticTimeManager->GetStaticTime(GetTimeOfDay());
 	}
-
-#if WITH_EDITORONLY_DATA
-	if (StaticTime.IsSet())
-	{
-		return StaticTime.GetValue();
-	}
-#endif
 	
 	return std::numeric_limits<float>::lowest();
 }
 
-void ADaySequenceActor::SetStaticTimeOfDay(float InHours)
+void ADaySequenceActor::RegisterStaticTimeContributor(const UE::DaySequence::FStaticTimeContributor& NewContributor) const
 {
-	using namespace UE::DaySequence;
-
-#if WITH_EDITORONLY_DATA
-	StaticTime = InHours;
-#endif
-	
-	// Convert the day time to sequence time
-	const FFrameRate FrameRate = RootSequence->GetMovieScene()->GetDisplayRate();
-	const float DayLengthHours = GetDayLength();
-	const float DayLengthRatio = FMath::Frac(InHours / DayLengthHours);
-	const float DayCycleSeconds = TimePerCycle.ToSeconds() * DayLengthRatio;
-	const FFrameTime StaticFrameTime = FrameRate.AsFrameTime(DayCycleSeconds);
-
-	if (TSharedPtr<FStaticTimeControllerOverride> ExistingTimeControllerOverride = WeakTimeControllerOverride.Pin())
+	if (!StaticTimeManager)
 	{
-		ExistingTimeControllerOverride->StaticTime = StaticFrameTime;
-		ExistingTimeControllerOverride->TimeInHours = InHours;
-	}
-	else
-	{
-		if (UDaySequencePlayer* Player = GetSequencePlayerInternal())
-		{
-			TSharedRef<FStaticTimeControllerOverride> TimeControllerOverride =
-				MakeShared<FStaticTimeControllerOverride>(StaticFrameTime, Player->GetTimeController());
-
-			TimeControllerOverride->TimeInHours = InHours;
-#if WITH_EDITORONLY_DATA
-			TimeControllerOverride->PreviousTimeOfDayPreview = TimeOfDayPreview;
-#endif
-
-			Player->SetTimeControllerDirectly(TimeControllerOverride);
-			Player->SetIgnorePlaybackReplication(true);
-
-			WeakTimeControllerOverride = TimeControllerOverride;
-		}
+		return;
 	}
-
-#if WITH_EDITOR
-	ConditionalSetTimeOfDayPreview(InHours);
-#endif
+	
+	StaticTimeManager->AddStaticTimeContributor(NewContributor);
 }
 
-void ADaySequenceActor::RemoveStaticTimeOfDay(bool bResumeFromStaticTime)
+void ADaySequenceActor::UnregisterStaticTimeContributor(const UObject* InUserObject) const
 {
-	using namespace UE::DaySequence;
-
-#if WITH_EDITORONLY_DATA
-	StaticTime.Reset();
-#endif
-	
-	TSharedPtr<FStaticTimeControllerOverride> TimeControllerOverride = WeakTimeControllerOverride.Pin();
-	UDaySequencePlayer* Player = GetSequencePlayerInternal();
-	if (TimeControllerOverride && Player)
+	if (!StaticTimeManager)
 	{
-		if (bResumeFromStaticTime && TimeControllerOverride->PreviousTimeController)
-		{
-			// Reset the old time controller to the static time if we want to resume from here
-			TimeControllerOverride->PreviousTimeController->Reset(Player->GetCurrentTime());
-		}
-
-		// Restore the old time controller and re-enable network replication
-		Player->SetTimeControllerDirectly(TimeControllerOverride->PreviousTimeController);
-		Player->SetIgnorePlaybackReplication(false);
-
-#if WITH_EDITOR
-		// Update the time of day preview if we're in an editor world
-		UWorld* World = GetWorld();
-		if (World && World->WorldType == EWorldType::Editor)
-		{
-			if (!bResumeFromStaticTime)
-			{
-				TimeOfDayPreview = TimeControllerOverride->PreviousTimeOfDayPreview;
-				OnTimeOfDayPreviewChangedEvent.Broadcast(GetTimeOfDayPreview());
-			}
-		}
-#endif
-
-		WeakTimeControllerOverride = nullptr;
+		return;
 	}
+	
+	StaticTimeManager->RemoveStaticTimeContributor(InUserObject);
 }
 
 void ADaySequenceActor::Play()
diff --git a/Engine/Plugins/Experimental/DaySequence/Source/DaySequence/Private/DaySequenceModifierComponent.cpp b/Engine/Plugins/Experimental/DaySequence/Source/DaySequence/Private/DaySequenceModifierComponent.cpp
--- a/Engine/Plugins/Experimental/DaySequence/Source/DaySequence/Private/DaySequenceModifierComponent.cpp
+++ b/Engine/Plugins/Experimental/DaySequence/Source/DaySequence/Private/DaySequenceModifierComponent.cpp
@@ -8,6 +8,7 @@
 #include "DaySequenceModule.h"
 #include "DaySequencePlayer.h"
 #include "DaySequenceTrack.h"
+#include "DaySequenceStaticTime.h"
 
 #include "Engine/World.h"
 #include "Materials/MaterialInterface.h"
@@ -673,7 +674,7 @@ void UDaySequenceModifierComponent::DisableModifier()
 		// Will call SetSubTrackMuteState for all living subsections, which checks enable state of modifier and their conditions
 		InvalidateMuteStates();
 
-		TargetActor->RemoveStaticTimeOfDay();
+		TargetActor->UnregisterStaticTimeContributor(this);
 		
 		if (bUnpauseOnDisable)
 		{
@@ -697,37 +698,46 @@ void UDaySequenceModifierComponent::SetInitialTimeOfDay()
 	{
 		const bool bHasAuthority = TargetActor->HasAuthority();
 		const bool bRandomTimeOfDay = DayNightCycle == EDayNightCycleMode::RandomFixedTime || DayNightCycle == EDayNightCycleMode::RandomStartTime;
-		const float Time = bRandomTimeOfDay ? FMath::FRand()*TargetActor->GetDayLength() : DayNightCycleTime;
-
-		if (!bHasAuthority && !bUseVolume)
-		{
-			// Never set initial time of day from non-volume based modifiers if they don't have authority
-			// We'll just get the initial time of day from the server replication
-			return;
-		}
+		const float RandomTime = FMath::FRand()*TargetActor->GetDayLength();
+		const float Time = bRandomTimeOfDay ? RandomTime : DayNightCycleTime;
 
 		switch (DayNightCycle)
 		{
 		case EDayNightCycleMode::FixedTime:
 		case EDayNightCycleMode::RandomFixedTime:
-			if (!bHasAuthority && bUseVolume)
 			{
-				// This function assigns a custom time controller so we can override
-				// the time regardless of server replication
-				TargetActor->SetStaticTimeOfDay(Time);
-				break;
-			}
+				auto WantsStaticTime = [this]() -> bool
+				{
+					return IsValid(this) && bIsEnabled && bIsComponentEnabled &&
+						(DayNightCycle == EDayNightCycleMode::FixedTime || DayNightCycle == EDayNightCycleMode::RandomFixedTime);
+				};
 
-			// If we're not overriding the time on a client, we need to make sure the time is replicated correctly
-			// AddStaticTimeOfDayOverride should be used where a static time of day needs to be evaluated from the sequence itself (to support h-bias overriding)
-			bUnpauseOnDisable = TargetActor->IsPlaying();
-			TargetActor->Pause();
+				auto GetStaticTime = [this, RandomTime, WantsStaticTime](UE::DaySequence::FStaticTimeInfo& OutRequest) -> bool
+				{
+					if (WantsStaticTime())
+					{
+						OutRequest.BlendWeight = bUseVolume ? GetCurrentBlendWeight() : 1.f;
+						OutRequest.StaticTime = DayNightCycle == EDayNightCycleMode::RandomFixedTime ? RandomTime : DayNightCycleTime;
+						return true;
+					}
 
-			// Intentional fallthrough - set the time and preview time
+					return false;
+				};
+			
+				TargetActor->RegisterStaticTimeContributor({this, Bias, WantsStaticTime, GetStaticTime});
+			}
+			break;
 
 		case EDayNightCycleMode::StartAtSpecifiedTime:
 		case EDayNightCycleMode::RandomStartTime:
 
+		if (!bHasAuthority && !bUseVolume)
+		{
+			// Never set initial time of day from non-volume based modifiers if they don't have authority and aren't setting static time.
+			// We'll just get the initial time of day from the server replication.
+			return;
+		}
+			
 			TargetActor->SetTimeOfDay(Time);
 #if WITH_EDITOR
 			TargetActor->ConditionalSetTimeOfDayPreview(Time);
diff --git a/Engine/Plugins/Experimental/DaySequence/Source/DaySequence/Private/DaySequenceModule.cpp b/Engine/Plugins/Experimental/DaySequence/Source/DaySequence/Private/DaySequenceModule.cpp
--- a/Engine/Plugins/Experimental/DaySequence/Source/DaySequence/Private/DaySequenceModule.cpp
+++ b/Engine/Plugins/Experimental/DaySequence/Source/DaySequence/Private/DaySequenceModule.cpp
@@ -13,39 +13,11 @@
 DEFINE_LOG_CATEGORY(LogDaySequence);
 CSV_DEFINE_CATEGORY(DaySequence, false);
 
-namespace UE::DaySequence
-{
-
-float GetStaticTimeOfDay(const UObject* Object)
-{
-	const ADaySequenceActor* DaySequenceActor = CastChecked<const ADaySequenceActor>(Object);
-	return DaySequenceActor->GetStaticTimeOfDay();
-}
-
-void SetStaticTimeOfDay(UObject* Object, float InStaticTimeOfDay)
-{
-	ADaySequenceActor* DaySequenceActor = CastChecked<ADaySequenceActor>(Object);
-	if (InStaticTimeOfDay == std::numeric_limits<float>::lowest())
-	{
-		DaySequenceActor->RemoveStaticTimeOfDay();
-	}
-	else
-	{
-		DaySequenceActor->SetStaticTimeOfDay(InStaticTimeOfDay);
-	}
-}
-
-} // namespace UE::DaySequence
-
 void FDaySequenceModule::StartupModule()
 {
 	using namespace UE::DaySequence;
 
 	OnCreateMovieSceneObjectSpawnerDelegateHandle = RegisterObjectSpawner(FOnCreateMovieSceneObjectSpawner::CreateStatic(&FDaySequenceActorSpawner::CreateObjectSpawner));
-
-	UE::MovieScene::FMovieSceneTracksComponentTypes::Get()->Accessors.Float.Add(
-		ADaySequenceActor::StaticClass(), "StaticTimeOfDay", &GetStaticTimeOfDay, &SetStaticTimeOfDay
-	);
 }
 
 void FDaySequenceModule::ShutdownModule()
diff --git a/Engine/Plugins/Experimental/DaySequence/Source/DaySequence/Private/DaySequenceStaticTime.cpp b/Engine/Plugins/Experimental/DaySequence/Source/DaySequence/Private/DaySequenceStaticTime.cpp
new file mode 100644
--- /dev/null
+++ b/Engine/Plugins/Experimental/DaySequence/Source/DaySequence/Private/DaySequenceStaticTime.cpp
@@ -0,0 +1,139 @@
+// Copyright Epic Games, Inc. All Rights Reserved.
+
+#pragma once
+
+#include "DaySequenceStaticTime.h"
+
+namespace UE::DaySequence
+{
+	void FStaticTimeManager::AddStaticTimeContributor(const FStaticTimeContributor& NewContributor)
+	{
+		if (!NewContributor.GetStaticTime || !NewContributor.WantsStaticTime || !NewContributor.UserObject.Get())
+		{
+			// Should be an ensure probably, we can't accept contributors that don't fulfill this requirement
+			return;
+		}
+
+		// If contributor is already registered, just early out.
+		if (Contributors.ContainsByPredicate([NewContributor](FStaticTimeContributor& Contributor)
+			{ return NewContributor.UserObject == Contributor.UserObject; }))
+		{
+			return;
+		}
+
+		// Add the new contributor.
+		Contributors.Add(NewContributor);
+		
+		// Sort the array in descending order by priority for efficient group processing.
+		Contributors.Sort([](const FStaticTimeContributor& LHS, const FStaticTimeContributor& RHS)
+			{ return LHS.Priority > RHS.Priority; });
+
+		// Increment (or set to 1) the priority group size.
+		PriorityGroupSizes.FindOrAdd(NewContributor.Priority, 0)++;
+	}
+
+	void FStaticTimeManager::RemoveStaticTimeContributor(const UObject* UserObject)
+	{
+		// Get the index of the contributor to remove.
+		int32 RemoveIdx = Contributors.IndexOfByPredicate([UserObject](const FStaticTimeContributor& Contributor)
+		{
+			return Contributor.UserObject.Get() == UserObject;
+		});
+
+		// If contributor actually exists, decrement its group count and remove it.
+		if (RemoveIdx != INDEX_NONE)
+		{
+			PriorityGroupSizes[Contributors[RemoveIdx].Priority]--;
+
+			Contributors.RemoveAt(RemoveIdx, EAllowShrinking::No);
+		}
+	}
+	
+	bool FStaticTimeManager::HasStaticTime() const
+	{
+		for (const FStaticTimeContributor& Contributor : Contributors)
+		{
+			if (Contributor.UserObject.Get() && Contributor.WantsStaticTime())
+			{
+				return true;
+			}
+		}
+
+		return false;
+	}
+	
+	float FStaticTimeManager::GetStaticTime(float InitialTime) const
+	{
+		float AccumulatedWeight = 0.f;
+		float AccumulatedTime = 0.f;
+			
+		// Process batches of contributors based on priority
+		for (int32 PriorityGroupStartIdx = 0; PriorityGroupStartIdx < Contributors.Num();)
+		{
+			// Compute the index of the final element in this priority group
+			const int32 CurrentPriority = Contributors[PriorityGroupStartIdx].Priority;
+			const int32 GroupSize = PriorityGroupSizes.FindChecked(CurrentPriority);	// checked because if this isn't valid something is very wrong
+			const int32 PriorityGroupEndIdx = PriorityGroupStartIdx + GroupSize - 1;
+			
+			// Process this group 
+			const FStaticTimeInfo GroupInfo = ProcessPriorityGroup(PriorityGroupStartIdx, PriorityGroupEndIdx);
+
+			const float EffectiveGroupWeight = (1.f - AccumulatedWeight) * GroupInfo.BlendWeight;
+			const float EffectiveGroupTime = EffectiveGroupWeight * GroupInfo.StaticTime;
+
+			AccumulatedWeight += EffectiveGroupWeight;
+			AccumulatedTime += EffectiveGroupTime;
+
+			// Advance to next priority group
+			PriorityGroupStartIdx = PriorityGroupEndIdx + 1;
+		}
+
+		// Blend against initial value if necessary
+		if (AccumulatedWeight < 1.f)
+		{
+			const float FillWeight = (1.f - AccumulatedWeight);
+			const float FillTime = InitialTime * FillWeight;
+
+			AccumulatedWeight += FillWeight;
+			AccumulatedTime += FillTime;
+		}
+
+		// we could probably ensure(AccumulatedWeight is nearly equal to 1.f), leaving out for now
+
+		return AccumulatedTime;
+	}
+
+	FStaticTimeInfo FStaticTimeManager::ProcessPriorityGroup(int32 StartIdx, int32 EndIdx) const
+	{
+		FStaticTimeInfo GroupInfo = {0.f, 0.f};
+		int32 GroupContributors = 0;
+		
+		for (int32 CurrentIdx = StartIdx; CurrentIdx <= EndIdx; ++CurrentIdx)
+		{
+			const FStaticTimeContributor& CurrentContributor = Contributors[CurrentIdx];
+
+			if (!CurrentContributor.UserObject.Get() || !CurrentContributor.WantsStaticTime())
+			{
+				continue;
+			}
+
+			// Only increment for active contributors
+			// Note: Because removing a contributor results in a discrete change in an integer value, we get pops when contributor in a group of >1 contributors has a non-1 weight.
+			GroupContributors++;
+			
+			FStaticTimeInfo ContributorInfo;
+			CurrentContributor.GetStaticTime(ContributorInfo);
+
+			// Accumulate contributor info (we will later divide these sums once we know how many contributors there are)
+			GroupInfo.BlendWeight += ContributorInfo.BlendWeight;
+			GroupInfo.StaticTime += ContributorInfo.StaticTime;
+		}
+
+		// Compute the average (before this point we have just summed all weights and times) for this group
+		// If GroupContributors is 0 we treat it as 1 because that implies the values in GroupInfo are just 0.
+		GroupInfo.BlendWeight /= FMath::Max(GroupContributors, 1);
+		GroupInfo.StaticTime /= FMath::Max(GroupContributors, 1);
+		
+		return GroupInfo;
+	}
+}
\ No newline at end of file
diff --git a/Engine/Plugins/Experimental/DaySequence/Source/DaySequence/Private/DaySequenceStaticTime.h b/Engine/Plugins/Experimental/DaySequence/Source/DaySequence/Private/DaySequenceStaticTime.h
new file mode 100644
--- /dev/null
+++ b/Engine/Plugins/Experimental/DaySequence/Source/DaySequence/Private/DaySequenceStaticTime.h
@@ -0,0 +1,59 @@
+// Copyright Epic Games, Inc. All Rights Reserved.
+
+#pragma once
+
+#include "CoreMinimal.h"
+#include "UObject/WeakObjectPtr.h"
+#include "Algo/Sort.h"
+
+namespace UE::DaySequence
+{
+	// The fundamental piece of data used by the static time system.
+	struct FStaticTimeInfo
+	{
+		float BlendWeight;
+		float StaticTime;
+	};
+
+	// The function signature for the static time requested callback. Must return true for a contributor to be considered
+	using FWantsStaticTimeFunction = TFunction<bool()>;
+	
+	// The function signature for the static time info getter callback. Only called if IsStaticTimeRequested returns true.
+	// This should return the same value that calling FStaticTimeRequestedFunction would return.
+	using FGetStaticTimeFunction = TFunction<bool(FStaticTimeInfo& OutRequest)>;
+	
+	// Contributors register an instance of this struct in order to request a static time.
+	struct FStaticTimeContributor
+	{
+		// Determines the lifetime of the object and prevents double registration.
+		TWeakObjectPtr<UObject> UserObject;
+
+		// Used for sorting contributors
+		int32 Priority;
+
+		// Returns whether or not this contributor is active.
+		FWantsStaticTimeFunction WantsStaticTime;
+
+		// Provides caller with desired static time information.
+		FGetStaticTimeFunction GetStaticTime;
+	};
+	
+	struct FStaticTimeManager
+	{
+		void AddStaticTimeContributor(const FStaticTimeContributor& NewContributor);
+		void RemoveStaticTimeContributor(const UObject* UserObject);
+		bool HasStaticTime() const;
+		float GetStaticTime(float InitialTime) const;
+		
+	private:
+
+		// Handles multiple contributors of the same priority by averaging their weights and times and returning a single desired weight and time
+		FStaticTimeInfo ProcessPriorityGroup(int32 StartIdx, int32 EndIdx) const;
+		
+		TArray<FStaticTimeContributor> Contributors;
+
+		// We precompute this when mutating the Contributors array to avoid redundant reiteration over Contributors when computing static time
+		TMap<int32, int32> PriorityGroupSizes;
+	};
+}
+
diff --git a/Engine/Plugins/Experimental/DaySequence/Source/DaySequence/Public/DaySequenceActor.h b/Engine/Plugins/Experimental/DaySequence/Source/DaySequence/Public/DaySequenceActor.h
--- a/Engine/Plugins/Experimental/DaySequence/Source/DaySequence/Public/DaySequenceActor.h
+++ b/Engine/Plugins/Experimental/DaySequence/Source/DaySequence/Public/DaySequenceActor.h
@@ -39,7 +39,8 @@ struct FMovieSceneSequencePlaybackSettings;
 
 namespace UE::DaySequence
 {
-	struct FStaticTimeControllerOverride;
+	struct FStaticTimeContributor;
+	struct FStaticTimeManager;
 
 #if ENABLE_DRAW_DEBUG
 	// This provides methods for determining if this debug entry should be shown and for getting a pointer to the debug data.
@@ -239,23 +240,9 @@ public:
 	UFUNCTION(BlueprintCallable, Category="TimeOfDay")
 	float GetStaticTimeOfDay() const;
 
-	/**
-	 * Force the Day Sequence to show a specific constant time (in hours)
-	 */
-	UFUNCTION(BlueprintCallable, Category="TimeOfDay")
-	void SetStaticTimeOfDay(float InHours);
-
-	/**
-	 * Revert the Day Sequence time back to its regular day cycle
-	 *
-	 * @param bResumeFromStaticTime  When checked, the day sequence will resume its playback from the static time,
-	 *                               otherwise it will restore to its natural position as if it had no static time applied.
-	 *                               Note: when this sequence is replicated, enabling this parameter will immediately be overwritten
-	 *                               by the server time on the next replication event.
-	 */
-	UFUNCTION(BlueprintCallable, Category="TimeOfDay")
-	void RemoveStaticTimeOfDay(bool bResumeFromStaticTime = false);
-
+	void RegisterStaticTimeContributor(const UE::DaySequence::FStaticTimeContributor& NewContributor) const;
+	void UnregisterStaticTimeContributor(const UObject* InUserObject) const;
+	
 	UDaySequence* GetRootSequence() const;
 
 	virtual void UpdateRootSequence();
@@ -507,16 +494,9 @@ protected:
 	UPROPERTY(EditAnywhere, Category = RuntimeDayCycle)
 	FDaySequenceTime InitialTimeOfDay;
 	
-	/** Cached static time controller applied from SetStaticTimeOfDay */
-	TWeakPtr<UE::DaySequence::FStaticTimeControllerOverride> WeakTimeControllerOverride;
+	TSharedPtr<UE::DaySequence::FStaticTimeManager> StaticTimeManager;
 
 #if WITH_EDITORONLY_DATA
-	/**
-	 * Editor only static time, necessary because WeakTimeControllerOverride only exists in
-	 * game worlds and we have no way of tracking static time outside of game worlds otherwise.
-	 */
-	TOptional<float> StaticTime;
-	
 	/**
 	 * Blueprint exposed delegate invoked when the TimeOfDayPreview property
 	 * is changed.
diff --git a/Engine/Plugins/Experimental/DaySequence/Source/DaySequence/Public/DaySequenceModifierComponent.h b/Engine/Plugins/Experimental/DaySequence/Source/DaySequence/Public/DaySequenceModifierComponent.h
--- a/Engine/Plugins/Experimental/DaySequence/Source/DaySequence/Public/DaySequenceModifierComponent.h
+++ b/Engine/Plugins/Experimental/DaySequence/Source/DaySequence/Public/DaySequenceModifierComponent.h
@@ -363,7 +363,7 @@ protected:
 protected:
 
 	/** Non-serialized target actor we are currently bound to */
-	UPROPERTY(Transient, BlueprintReadWrite, Category="Day Sequence")
+	UPROPERTY(Transient, DuplicateTransient, BlueprintReadWrite, Category="Day Sequence")
 	TObjectPtr<ADaySequenceActor> TargetActor;
 
 	/** A handle used to force an override of the TargetActor's evaluation interval. */
