DaySequenceStaticTime: Fix non unity error

Accidentally copied a #pragma once into a cpp while shuffling code around

diff --git a/Engine/Plugins/Experimental/DaySequence/Source/DaySequence/Private/DaySequenceStaticTime.cpp b/Engine/Plugins/Experimental/DaySequence/Source/DaySequence/Private/DaySequenceStaticTime.cpp
--- a/Engine/Plugins/Experimental/DaySequence/Source/DaySequence/Private/DaySequenceStaticTime.cpp
+++ b/Engine/Plugins/Experimental/DaySequence/Source/DaySequence/Private/DaySequenceStaticTime.cpp
@@ -1,7 +1,5 @@
 // Copyright Epic Games, Inc. All Rights Reserved.
 
-#pragma once
-
 #include "DaySequenceStaticTime.h"
 
 namespace UE::DaySequence
