DaySequenceStaticTime: Add a blueprint exposed static time contributor object.

This object allows BPs (or really anything) to create an instance of this object and bind it to a DSA. It handles the creation of an FStaticTimeContributor and manages the lifetime. Users can unregister the contributor manually or it will be removed when the instance of this object is destroyed

#changelist validated
#virtualized

diff --git a/Engine/Plugins/Experimental/DaySequence/Source/DaySequence/Private/DaySequenceStaticTime.cpp b/Engine/Plugins/Experimental/DaySequence/Source/DaySequence/Private/DaySequenceStaticTime.cpp
--- a/Engine/Plugins/Experimental/DaySequence/Source/DaySequence/Private/DaySequenceStaticTime.cpp
+++ b/Engine/Plugins/Experimental/DaySequence/Source/DaySequence/Private/DaySequenceStaticTime.cpp
@@ -12,12 +12,17 @@ namespace UE::DaySequence
 			return;
 		}
 
-		// If contributor is already registered, just early out.
-		if (Contributors.ContainsByPredicate([NewContributor](FStaticTimeContributor& Contributor)
-			{ return NewContributor.UserObject == Contributor.UserObject; }))
+		// We remove an existing matching contributor if necessary, making sure to keep PriorityGroupSizes up to date.
+		Contributors.RemoveAll([this, NewContributor](const FStaticTimeContributor& Contributor)
 		{
-			return;
-		}
+			if (NewContributor.UserObject == Contributor.UserObject)
+			{
+				PriorityGroupSizes[NewContributor.Priority]--;
+				return true;
+			}
+			
+			return false;
+		});
 
 		// Add the new contributor.
 		Contributors.Add(NewContributor);
@@ -134,4 +139,62 @@ namespace UE::DaySequence
 		
 		return GroupInfo;
 	}
+}
+
+UDaySequenceStaticTimeContributor::UDaySequenceStaticTimeContributor()
+: BlendWeight(1.f)
+, StaticTime(0.f)
+, bWantsStaticTime(true)
+, TargetActor(nullptr)
+{}
+
+void UDaySequenceStaticTimeContributor::BeginDestroy()
+{
+	UnbindFromDaySequenceActor();
+	
+	Super::BeginDestroy();
+}
+
+void UDaySequenceStaticTimeContributor::BindToDaySequenceActor(ADaySequenceActor* InTargetActor, int32 Priority)
+{
+	UnbindFromDaySequenceActor();
+
+	if (!InTargetActor)
+	{
+		return;
+	}
+		
+	TargetActor = InTargetActor;
+
+	UObject* Outer = GetOuter();
+	auto WantsStaticTime = [this, Outer]()
+	{
+		return IsValid(this) && IsValid(Outer) && bWantsStaticTime;
+	};
+
+	auto GetStaticTime = [this, WantsStaticTime](UE::DaySequence::FStaticTimeInfo& OutRequest)
+	{
+		if (WantsStaticTime())
+		{
+			OutRequest.BlendWeight = BlendWeight;
+			OutRequest.StaticTime = StaticTime;
+			return true;
+		}
+
+		return false;
+	};
+
+	TargetActor->RegisterStaticTimeContributor({ Outer, Priority, WantsStaticTime, GetStaticTime });
+}
+
+void UDaySequenceStaticTimeContributor::UnbindFromDaySequenceActor()
+{
+	if (!TargetActor)
+	{
+		return;
+	}
+
+	TargetActor->UnregisterStaticTimeContributor(this);
+
+	TargetActor = nullptr;
 }
\ No newline at end of file
diff --git a/Engine/Plugins/Experimental/DaySequence/Source/DaySequence/Private/DaySequenceStaticTime.h b/Engine/Plugins/Experimental/DaySequence/Source/DaySequence/Private/DaySequenceStaticTime.h
--- a/Engine/Plugins/Experimental/DaySequence/Source/DaySequence/Private/DaySequenceStaticTime.h
+++ b/Engine/Plugins/Experimental/DaySequence/Source/DaySequence/Private/DaySequenceStaticTime.h
@@ -5,6 +5,9 @@
 #include "CoreMinimal.h"
 #include "UObject/WeakObjectPtr.h"
 #include "Algo/Sort.h"
+#include "DaySequenceActor.h"
+
+#include "DaySequenceStaticTime.generated.h"
 
 namespace UE::DaySequence
 {
@@ -57,3 +60,43 @@ namespace UE::DaySequence
 	};
 }
 
+/**
+ * A Blueprint exposed static time contributor.
+ * Used to contribute to static time blending for the specified Day Sequence Actor without needing to spawn actors and/or components.
+ */
+UCLASS(BlueprintType)
+class UDaySequenceStaticTimeContributor : public UObject
+{
+	GENERATED_BODY()
+
+public:
+
+	UDaySequenceStaticTimeContributor();
+	
+	virtual void BeginDestroy() override;
+	
+	/** The desired blend weight. Once bound to a Day Sequence Actor, this can be freely changed without rebinding. */
+	UPROPERTY(Interp, BlueprintReadWrite, Category = "Day Sequence")
+	float BlendWeight;
+
+	/** The desired static time. Once bound to a Day Sequence Actor, this can be freely changed without rebinding. */
+	UPROPERTY(Interp, BlueprintReadWrite, Category = "Day Sequence")
+	float StaticTime;
+
+	/** Determines whether or not this contributor is effective once we are bound. This can be freely changed to enable/disable the contributor without rebinding. */
+	UPROPERTY(Interp, BlueprintReadWrite, Category = "Day Sequence")
+	bool bWantsStaticTime;
+
+	/** Begin contributing static time to the specified actor. */
+	UFUNCTION(BlueprintCallable, Category = "Day Sequence")
+	void BindToDaySequenceActor(ADaySequenceActor* InTargetActor, int32 Priority = 1000);
+
+	/** Stop contributing static time. */
+	UFUNCTION(BlueprintCallable, Category = "Day Sequence")
+	void UnbindFromDaySequenceActor();
+	
+private:
+
+	UPROPERTY(Transient)
+	TObjectPtr<ADaySequenceActor> TargetActor;
+};	
