DaySequence: Fix section locking

diff --git a/Engine/Plugins/Experimental/DaySequence/Source/DaySequence/Private/DaySequenceActor.cpp b/Engine/Plugins/Experimental/DaySequence/Source/DaySequence/Private/DaySequenceActor.cpp
--- a/Engine/Plugins/Experimental/DaySequence/Source/DaySequence/Private/DaySequenceActor.cpp
+++ b/Engine/Plugins/Experimental/DaySequence/Source/DaySequence/Private/DaySequenceActor.cpp
@@ -721,14 +721,17 @@ UMovieSceneSubSection* ADaySequenceActor::InitializeDaySequence(const FDaySequen
 					return;
 				}
 
+				SubSection->SetIsLocked(false);
+				// Begin SubSection mutation:
+				
 				constexpr bool bInitialMuteState = false;
 				if (const bool bActive = !EvaluateSequenceConditions(bInitialMuteState, Conditions); SubSection->IsActive() != bActive)
 				{
-					SubSection->SetIsLocked(false);
 					SubSection->MarkAsChanged();
 					SubSection->SetIsActive(bActive);
-					SubSection->SetIsLocked(true);
 				}
+
+				SubSection->SetIsLocked(true);
 			};
 
 			const TFunction<void(void)> SetSubTrackMuteStateUnconditional = [this, SubSection]()
@@ -738,13 +741,16 @@ UMovieSceneSubSection* ADaySequenceActor::InitializeDaySequence(const FDaySequen
 					return;
 				}
 
+				SubSection->SetIsLocked(false);
+				// Begin SubSection mutation:
+				
 				if (!SubSection->IsActive())
 				{
-					SubSection->SetIsLocked(false);
 					SubSection->MarkAsChanged();
 					SubSection->SetIsActive(true);
-					SubSection->SetIsLocked(true);
 				}
+				
+				SubSection->SetIsLocked(true);
 			};
 			
 			const TFunction<void(void)>& SetSubTrackMuteState = SequenceAsset.Conditions.Conditions.Num() == 0 ? SetSubTrackMuteStateUnconditional : SetSubTrackMuteStateConditional;
diff --git a/Engine/Plugins/Experimental/DaySequence/Source/DaySequence/Private/DaySequenceModifierComponent.cpp b/Engine/Plugins/Experimental/DaySequence/Source/DaySequence/Private/DaySequenceModifierComponent.cpp
--- a/Engine/Plugins/Experimental/DaySequence/Source/DaySequence/Private/DaySequenceModifierComponent.cpp
+++ b/Engine/Plugins/Experimental/DaySequence/Source/DaySequence/Private/DaySequenceModifierComponent.cpp
@@ -922,15 +922,18 @@ UMovieSceneSubSection* UDaySequenceModifierComponent::InitializeDaySequence(cons
 			return;
 		}
 
+		SubSection->SetIsLocked(false);
+		// Begin SubSection mutation:
+		
 		constexpr bool bInitialMuteState = false;
 		const bool bActive = bIsEnabled && !TargetActor->EvaluateSequenceConditions(bInitialMuteState, Conditions);
 		if (SubSection->IsActive() != bActive)
 		{
-			SubSection->SetIsLocked(false);
 			SubSection->MarkAsChanged();
 			SubSection->SetIsActive(bActive);
-			SubSection->SetIsLocked(true);
 		}
+
+		SubSection->SetIsLocked(true);
 	};
 
 	const TFunction<void(void)> SetSubTrackMuteStateUnconditional = [this, SubSection]()
@@ -940,14 +943,17 @@ UMovieSceneSubSection* UDaySequenceModifierComponent::InitializeDaySequence(cons
 			return;
 		}
 
+		SubSection->SetIsLocked(false);
+		// Begin SubSection mutation:
+		
 		const bool bActive = bIsEnabled;
 		if (SubSection->IsActive() != bActive)
 		{
-			SubSection->SetIsLocked(false);
 			SubSection->MarkAsChanged();
 			SubSection->SetIsActive(bActive);
-			SubSection->SetIsLocked(true);
 		}
+
+		SubSection->SetIsLocked(true);
 	};
 
 	const TFunction<void(void)>& SetSubTrackMuteState = Entry.Conditions.Conditions.Num() == 0 ? SetSubTrackMuteStateUnconditional : SetSubTrackMuteStateConditional;
