DaySequence: Set reasonable defaults in procedural sequences

diff --git a/Engine/Plugins/Experimental/DaySequence/Source/DaySequence/Public/ProceduralSequences/SunPositionSequence.h b/Engine/Plugins/Experimental/DaySequence/Source/DaySequence/Public/ProceduralSequences/SunPositionSequence.h
--- a/Engine/Plugins/Experimental/DaySequence/Source/DaySequence/Public/ProceduralSequences/SunPositionSequence.h
+++ b/Engine/Plugins/Experimental/DaySequence/Source/DaySequence/Public/ProceduralSequences/SunPositionSequence.h
@@ -27,16 +27,16 @@ struct DAYSEQUENCE_API FSunPositionSequence : public FProceduralDaySequence
 	FDateTime Time;
 
 	UPROPERTY(EditAnywhere, Category = "Procedural Parameters")
-	double TimeZone;
+	double TimeZone = 0.0;
 	
 	UPROPERTY(EditAnywhere, Category = "Procedural Parameters")
-	double Latitude;
+	double Latitude = 0.0;
 
 	UPROPERTY(EditAnywhere, Category = "Procedural Parameters")
-	double Longitude;
+	double Longitude = 0.0;
 	
 	UPROPERTY(EditAnywhere, Category = "Procedural Parameters")
-	bool bIsDaylightSavings;
+	bool bIsDaylightSavings = false;
 	
 private:
 	
