DaySequenceModifierComponent: Rework blend mode & volume API

#changelist validated
#virtualized

diff --git a/Engine/Plugins/Experimental/DaySequence/Config/Engine.ini b/Engine/Plugins/Experimental/DaySequence/Config/Engine.ini
--- a/Engine/Plugins/Experimental/DaySequence/Config/Engine.ini
+++ b/Engine/Plugins/Experimental/DaySequence/Config/Engine.ini
@@ -9,4 +9,4 @@
 +ClassRedirects=(OldName="/Script/TimeOfDay.TimeOfDaySequenceTrack", NewName="/Script/DaySequence.DaySequenceTrack")
 +ClassRedirects=(OldName="/Script/TimeOfDay.TimeOfDayActor", NewName="/Script/DaySequence.DaySequenceActor")
 +ClassRedirects=(OldName="/Script/TimeOfDay.EnvironmentLightingActor", NewName="/Script/DaySequence.EnvironmentLightingActor")
-
++EnumRedirects=(OldName="/Script/DaySequence.EDaySequenceModifierBlendMode", NewName="/Script/DaySequence.EDaySequenceModifierMode", ValueChanges=(("None","Global"),("Distance","Volume")))
diff --git a/Engine/Plugins/Experimental/DaySequence/Source/DaySequence/Private/Actors/DaySequenceModifierVolume.cpp b/Engine/Plugins/Experimental/DaySequence/Source/DaySequence/Private/Actors/DaySequenceModifierVolume.cpp
--- a/Engine/Plugins/Experimental/DaySequence/Source/DaySequence/Private/Actors/DaySequenceModifierVolume.cpp
+++ b/Engine/Plugins/Experimental/DaySequence/Source/DaySequence/Private/Actors/DaySequenceModifierVolume.cpp
@@ -24,18 +24,11 @@ ADaySequenceModifierVolume::ADaySequenceModifierVolume(const FObjectInitializer&
 	DefaultBox->SetupAttachment(DaySequenceModifier);
 	DefaultBox->SetLineThickness(10.f);
 	DefaultBox->SetBoxExtent(FVector(500.f));
-}
 
-void ADaySequenceModifierVolume::TryEnableModifier() const
-{	
-	if (DaySequenceModifier->IsBlendTargetInAnyVolume())
-	{
-		DaySequenceModifier->EnableModifier();
-	}
-	else
-	{
-		DaySequenceModifier->DisableModifier();
-	}
+	FComponentReference DefaultBoxReference;
+	DefaultBoxReference.OverrideComponent = DefaultBox;
+	DaySequenceModifier->EmptyVolumeShapeComponents();
+	DaySequenceModifier->AddVolumeShapeComponent(DefaultBoxReference);
 }
 
 void ADaySequenceModifierVolume::SetBlendTarget(APlayerController* InPC)
@@ -47,10 +40,8 @@ void ADaySequenceModifierVolume::SetBlendTarget(APlayerController* InPC)
 
 	CurrentBlendTarget = InPC;
 
-	DaySequenceModifier->EnableDistanceVolumeBlends(InPC);
-	DaySequenceModifier->SetCustomVolumeBlendWeight(1.f);
-
-	TryEnableModifier();
+	DaySequenceModifier->SetBlendTarget(InPC);
+	DaySequenceModifier->SetUserBlendWeight(1.f);
 }
 
 void ADaySequenceModifierVolume::BeginPlay()
@@ -96,23 +87,6 @@ void ADaySequenceModifierVolume::Initialize()
 	DaySequenceActorSetup();
 }
 
-void ADaySequenceModifierVolume::VolumeSetup()
-{
-	DaySequenceModifier->EmptyVolumeShapeComponents();
-	
-	AddShapeComponentsToModifier();
-
-	// this is kind of a hack, we are calling this for the side effect of forcing the modifier to recache volume shapes after we just invalidated them.
-	DaySequenceModifier->IsBlendTargetInAnyVolume();
-}
-
-void ADaySequenceModifierVolume::AddShapeComponentsToModifier()
-{
-	FComponentReference DefaultBoxReference;
-	DefaultBoxReference.OverrideComponent = DefaultBox;
-	DaySequenceModifier->AddVolumeShapeComponent(DefaultBoxReference);
-}
-
 void ADaySequenceModifierVolume::PlayerControllerSetup()
 {
 #if WITH_EDITOR
@@ -199,7 +173,6 @@ void ADaySequenceModifierVolume::BindToDaySequenceActor()
 	{
 		DaySequenceModifier->BindToDaySequenceActor(DaySequenceActor);
 
-		VolumeSetup();
 		PlayerControllerSetup();
 
 		OnDaySequenceActorBound(DaySequenceActor);
diff --git a/Engine/Plugins/Experimental/DaySequence/Source/DaySequence/Private/DaySequenceModifierComponent.cpp b/Engine/Plugins/Experimental/DaySequence/Source/DaySequence/Private/DaySequenceModifierComponent.cpp
--- a/Engine/Plugins/Experimental/DaySequence/Source/DaySequence/Private/DaySequenceModifierComponent.cpp
+++ b/Engine/Plugins/Experimental/DaySequence/Source/DaySequence/Private/DaySequenceModifierComponent.cpp
@@ -219,14 +219,14 @@ void UDaySequenceModifierEasingFunction::Initialize(EEasingFunctionType EasingTy
 		case EEasingFunctionType::EaseIn:
 			EvaluateImpl = [Outer](float)
 			{
-				return Outer->GetCurrentBlendWeight();
+				return Outer->GetBlendWeight();
 			};
 			break;
 	
 		case EEasingFunctionType::EaseOut:
 			EvaluateImpl = [Outer](float)
 			{
-				return 1.f - Outer->GetCurrentBlendWeight();
+				return 1.f - Outer->GetBlendWeight();
 			};
 			break;
 		}
@@ -249,7 +249,6 @@ UDaySequenceModifierComponent::UDaySequenceModifierComponent(const FObjectInitia
 	bIsComponentEnabled = true;
 	bIsEnabled = false;
 	bIgnoreBias = false;
-	bUseVolume = true;
 	bPreview = true;
 	bUseCollection = false;
 	bSmoothBlending = false;
@@ -257,10 +256,11 @@ UDaySequenceModifierComponent::UDaySequenceModifierComponent(const FObjectInitia
 	Bias = 1000;
 	DayNightCycleTime = 12.f;
 	DayNightCycle = EDayNightCycleMode::Default;
-	BlendMode = EDaySequenceModifierBlendMode::Distance;
+	Mode = EDaySequenceModifierMode::Volume;
+	BlendPolicy = EDaySequenceModifierUserBlendPolicy::Minimum;
 	BlendAmount = 100.f;
-	CachedBlendFactor = 0.f;
-	CustomVolumeBlendWeight = 1.0f;
+	UserBlendWeight = 1.f;
+	InternalBlendWeight = 1.f;
 #if ENABLE_DRAW_DEBUG
 	DebugLevel = 0;
 	
@@ -273,9 +273,9 @@ UDaySequenceModifierComponent::UDaySequenceModifierComponent(const FObjectInitia
 		(*DebugData).FindOrAdd("Owner Name") = GetOwner()->GetFName().ToString();
 		(*DebugData).FindOrAdd("Component Enabled") = bIsComponentEnabled ? "True" : "False";
 		(*DebugData).FindOrAdd("Modifier Enabled") = bIsEnabled ? "True" : "False";
-		(*DebugData).FindOrAdd("Blend Weight") = FString::Printf(TEXT("%.5f"), GetCurrentBlendWeight());
+		(*DebugData).FindOrAdd("Blend Weight") = FString::Printf(TEXT("%.5f"), GetBlendWeight());
 
-		const APlayerController* BlendTarget = ExternalVolumeBlendTarget.Get();
+		const APlayerController* BlendTarget = WeakBlendTarget.Get();
 		(*DebugData).FindOrAdd("Blend Target" ) = BlendTarget ? BlendTarget->GetName() : "None";
 
 		return DebugData;
@@ -303,17 +303,16 @@ void UDaySequenceModifierComponent::UpdateEditorPreview(float DeltaTime)
 {
 	using namespace UE::DaySequence;
 
-	if (bIsComponentEnabled && bPreview && bUseVolume && IsRegistered() && !GetWorld()->IsGameWorld())
+	if (bIsComponentEnabled && bPreview && IsRegistered() && !GetWorld()->IsGameWorld())
 	{
-		if (const float DistanceBlendFactor = GetDistanceBlendFactor(GVolumePreviewLocation); DistanceBlendFactor > UE_SMALL_NUMBER)
+		const float OldEffectiveBlendWeight = GetBlendWeight();
+		
+		if (UpdateInternalBlendWeight() > UE_SMALL_NUMBER)
 		{
-			if (!bIsEnabled)
-			{
-				EnableModifier();
-			}
+			EnableModifier();
 			
-			const float BlendFactor = FMath::Min(DistanceBlendFactor, CustomVolumeBlendWeight);
-			if (BlendMode != EDaySequenceModifierBlendMode::None && CachedBlendFactor != BlendFactor)
+			// this compares effective blend weights, which is necessary in case user blend weight is changing while the internal blend weight remains constant.
+			if (!FMath::IsNearlyEqual(OldEffectiveBlendWeight, GetBlendWeight()))
 			{
 				// If we're using a blend we have to mark active sections as changed
 				// in order to force an update in-editor:
@@ -333,7 +332,6 @@ void UDaySequenceModifierComponent::UpdateEditorPreview(float DeltaTime)
 					}
 				}
 			}
-			CachedBlendFactor = BlendFactor;
 		}
 		else
 		{
@@ -374,19 +372,9 @@ void UDaySequenceModifierComponent::PostEditChangeProperty(FPropertyChangedEvent
 			DisableModifier();
 		}
 	}
-	else if (PropertyName == GET_MEMBER_NAME_CHECKED(UDaySequenceModifierComponent, bUseVolume))
+	else if (PropertyName == GET_MEMBER_NAME_CHECKED(UDaySequenceModifierComponent, Mode))
 	{
-		if (!bUseVolume)
-		{
-			if (bPreview && !bIsEnabled)
-			{
-				EnableModifier();
-			}
-			else if (!bPreview && bIsEnabled)
-			{
-				DisableModifier();
-			}
-		}
+		EnableModifier();
 	}
 }
 
@@ -397,11 +385,6 @@ void UDaySequenceModifierComponent::OnRegister()
 	Super::OnRegister();
 
 	bCachedExternalShapesInvalid = true;
-	
-	if (!bUseVolume)
-	{
-		EnableModifier();
-	}
 }
 
 void UDaySequenceModifierComponent::OnUnregister()
@@ -418,26 +401,14 @@ void UDaySequenceModifierComponent::DaySequenceUpdate()
 {
 	CSV_SCOPED_TIMING_STAT(DaySequence, SequencePlayerUpdated);
 	
-	// Force expensive update
-	const float DistanceBlendFactor = UpdateBlendWeight();
-
-	// This block begins overriding the update interval if necessary.
-	if (OverrideUpdateIntervalHandle)
-	{
-		if (bSmoothBlending &&
-			FMath::IsWithin(DistanceBlendFactor, 0.f + UE_KINDA_SMALL_NUMBER, 1.f - UE_KINDA_SMALL_NUMBER))
-		{
-			OverrideUpdateIntervalHandle->StartOverriding();
-		}
-		else
-		{
-			OverrideUpdateIntervalHandle->StopOverriding();
-		}
-	}
-	
-	if (bIsComponentEnabled && bUseVolume)
+	if (bIsComponentEnabled)
 	{
-		if (DistanceBlendFactor > UE_SMALL_NUMBER)
+		// Force the expensive update
+		UpdateInternalBlendWeight();
+		
+		// For the purposes of enable/disable we ignore BlendPolicy and directly use InternalBlendWeight,
+		// but the easing function will respect it (by calling GetBlendWeight).
+		if (InternalBlendWeight > UE_SMALL_NUMBER)
 		{
 			EnableModifier();
 		}
@@ -445,6 +416,19 @@ void UDaySequenceModifierComponent::DaySequenceUpdate()
 		{
 			DisableModifier();
 		}
+
+		if (OverrideUpdateIntervalHandle)
+		{
+			if (bIsEnabled && bSmoothBlending &&
+			FMath::IsWithin(GetBlendWeight(), 0.f + UE_KINDA_SMALL_NUMBER, 1.f - UE_KINDA_SMALL_NUMBER))
+			{
+				OverrideUpdateIntervalHandle->StartOverriding();
+			}
+			else
+			{
+				OverrideUpdateIntervalHandle->StopOverriding();
+			}
+		}
 	}
 }
 
@@ -589,7 +573,7 @@ bool UDaySequenceModifierComponent::CanBeEnabled() const
 		return false;
 	}
 	
-	if (bUseVolume)
+	if (Mode == EDaySequenceModifierMode::Volume)
 	{
 		ENetMode NetMode = Actor->GetNetMode();
 		return NetMode != NM_DedicatedServer;
@@ -687,9 +671,6 @@ void UDaySequenceModifierComponent::DisableModifier()
 			TargetActor->SetTimeOfDay(TargetActor->GetTimeOfDay());
 		}
 	}
-
-	// Necessary for correctly marking TargetActor as changed on enable.
-	CachedBlendFactor = -1.0f;
 }
 
 void UDaySequenceModifierComponent::SetInitialTimeOfDay()
@@ -716,7 +697,7 @@ void UDaySequenceModifierComponent::SetInitialTimeOfDay()
 				{
 					if (WantsStaticTime())
 					{
-						OutRequest.BlendWeight = bUseVolume ? GetCurrentBlendWeight() : 1.f;
+						OutRequest.BlendWeight = Mode == EDaySequenceModifierMode::Volume ? GetBlendWeight() : 1.f;
 						OutRequest.StaticTime = DayNightCycle == EDayNightCycleMode::RandomFixedTime ? RandomTime : DayNightCycleTime;
 						return true;
 					}
@@ -731,7 +712,7 @@ void UDaySequenceModifierComponent::SetInitialTimeOfDay()
 		case EDayNightCycleMode::StartAtSpecifiedTime:
 		case EDayNightCycleMode::RandomStartTime:
 
-		if (!bHasAuthority && !bUseVolume)
+		if (!bHasAuthority && Mode != EDaySequenceModifierMode::Volume)
 		{
 			// Never set initial time of day from non-volume based modifiers if they don't have authority and aren't setting static time.
 			// We'll just get the initial time of day from the server replication.
@@ -879,18 +860,14 @@ UMovieSceneSubSection* UDaySequenceModifierComponent::InitializeDaySequence(cons
 
 		RootTrack->AddSection(*SubSection);
 
-		if (bUseVolume && BlendMode != EDaySequenceModifierBlendMode::None && bBlendHierarchicalBias)
-		{
-			// In the Sequencer Editor, EaseIn pads the Sequence asset name by the EaseIn duration
-			// (see SSequencerSection::OnPaint). Since we set the Easing duration to the full section
-			// width to facilitate blending, the label is clipped. So we use EaseOut here instead and
-			// ensure that the weight is inverted in Evaluate().
-			SubSection->Easing.bManualEaseOut = true;
-			SubSection->Easing.ManualEaseOutDuration = MovieScene->GetPlaybackRange().Size<FFrameNumber>().Value;
-
-			EasingFunction->Initialize(UDaySequenceModifierEasingFunction::EEasingFunctionType::EaseOut);
-			SubSection->Easing.EaseOut = EasingFunction;
-		}
+		// In the Sequencer Editor, EaseIn pads the Sequence asset name by the EaseIn duration
+		// (see SSequencerSection::OnPaint). Since we set the Easing duration to the full section
+		// width to facilitate blending, the label is clipped. So we use EaseOut here instead and
+		// ensure that the weight is inverted in Evaluate().
+		SubSection->Easing.bManualEaseOut = true;
+		SubSection->Easing.ManualEaseOutDuration = MovieScene->GetPlaybackRange().Size<FFrameNumber>().Value;
+		EasingFunction->Initialize(UDaySequenceModifierEasingFunction::EEasingFunctionType::EaseOut);
+		SubSection->Easing.EaseOut = EasingFunction;
 
 #if WITH_EDITOR
 		FString Label = GetOwner()->GetActorLabel();
@@ -1388,6 +1365,16 @@ void UDaySequenceModifierComponent::SetUserDaySequence(UDaySequence* InDaySequen
 	ReinitializeSubSequence(nullptr);
 }
 
+void UDaySequenceModifierComponent::SetMode(EDaySequenceModifierMode NewMode)
+{
+	Mode = NewMode;
+}
+
+void UDaySequenceModifierComponent::SetBlendPolicy(EDaySequenceModifierUserBlendPolicy NewPolicy)
+{
+	BlendPolicy = NewPolicy;
+}
+
 bool UDaySequenceModifierComponent::GetBlendPosition(FVector& InPosition) const
 {
 	CSV_SCOPED_TIMING_STAT(DaySequence, GetBlendPosition);
@@ -1400,7 +1387,7 @@ bool UDaySequenceModifierComponent::GetBlendPosition(FVector& InPosition) const
 	}
 	else
 #endif
-	if (const APlayerController* BlendTarget = ExternalVolumeBlendTarget.Get(); BlendTarget && BlendTarget->PlayerCameraManager)
+	if (const APlayerController* BlendTarget = WeakBlendTarget.Get(); BlendTarget && BlendTarget->PlayerCameraManager)
 	{
 		CSV_SCOPED_TIMING_STAT(DaySequence, GetPlayerViewPoint);
 		InPosition = BlendTarget->PlayerCameraManager->GetCameraLocation();
@@ -1410,26 +1397,6 @@ bool UDaySequenceModifierComponent::GetBlendPosition(FVector& InPosition) const
 	return false;
 }
 
-float UDaySequenceModifierComponent::GetDistanceBlendFactorForShape(const UShapeComponent* Shape, const FVector& Position) const
-{
-	const float Distance = UE::DaySequence::ComputeSignedDistance(Shape, Position);
-	return Distance < 0.f ? FMath::Clamp(-Distance / BlendAmount, 0.f, 1.f) : 0.f;
-}
-
-float UDaySequenceModifierComponent::GetDistanceBlendFactor(const FVector& Position) const
-{
-	CSV_SCOPED_TIMING_STAT(DaySequence, GetDistanceBlendFactor);
-	
-	CachedDistanceBlendFactor = 0.f;
-
-	for (const UShapeComponent* Shape : GetVolumeShapeComponents())
-	{
-		CachedDistanceBlendFactor = FMath::Max(CachedDistanceBlendFactor, GetDistanceBlendFactorForShape(Shape, Position));
-	}
-
-	return CachedDistanceBlendFactor;
-}
-
 TArray<UShapeComponent*> UDaySequenceModifierComponent::GetVolumeShapeComponents() const
 {
 	TArray<UShapeComponent*> ResolvedVolumeShapeComponents;
@@ -1481,33 +1448,60 @@ TArray<UShapeComponent*> UDaySequenceModifierComponent::GetVolumeShapeComponents
 	return ResolvedVolumeShapeComponents;
 }
 
-float UDaySequenceModifierComponent::GetCurrentBlendWeight() const
+float UDaySequenceModifierComponent::GetBlendWeight() const
 {
-	return FMath::Min(CachedDistanceBlendFactor, CustomVolumeBlendWeight);
+	switch (BlendPolicy)
+	{
+	default:
+	case EDaySequenceModifierUserBlendPolicy::Ignored:
+		return InternalBlendWeight;
+		
+	case EDaySequenceModifierUserBlendPolicy::Minimum:
+		return FMath::Min(InternalBlendWeight, UserBlendWeight);
+		
+	case EDaySequenceModifierUserBlendPolicy::Maximum:
+		return FMath::Max(InternalBlendWeight, UserBlendWeight);
+		
+	case EDaySequenceModifierUserBlendPolicy::Override:
+		return UserBlendWeight;
+	}
 }
 
-float UDaySequenceModifierComponent::UpdateBlendWeight() const
+float UDaySequenceModifierComponent::UpdateInternalBlendWeight()
 {
-	FVector BlendPosition;
-	const bool bHasBlendPosition = GetBlendPosition(BlendPosition);
-
-	const float OldBlendWeight = CachedDistanceBlendFactor;
-	const float NewBlendWeight = FMath::Min(bHasBlendPosition ? GetDistanceBlendFactor(BlendPosition) : 1.f, CustomVolumeBlendWeight);
-
-	// todo [nickolas.drake]: Enable blending for paused actors. Need to force set time of day if:
-	// 1) We have a blend position
-	// 2) The target DSA is valid and not playing
-	// 3) Our old blend weight is sufficiently different from our new blend weight
+	CSV_SCOPED_TIMING_STAT(DaySequence, UpdateInternalBlendWeight);
 	
-	return NewBlendWeight;
-}
-
-void UDaySequenceModifierComponent::SetVolumeCollisionEnabled(const ECollisionEnabled::Type InCollisionType) const
-{
-	for (UShapeComponent* Shape : GetVolumeShapeComponents())
+	switch (Mode)
 	{
-		Shape->SetCollisionEnabled(InCollisionType);
+	default:
+	case EDaySequenceModifierMode::Global:
+		InternalBlendWeight = 1.f;
+		break;
+	
+	case EDaySequenceModifierMode::Volume:
+		if (FVector BlendPosition; GetBlendPosition(BlendPosition))
+		{
+			InternalBlendWeight = 0.f;
+
+			auto GetBlendWeightForShape = [this, BlendPosition](const UShapeComponent* Shape)
+			{
+				const float Distance = UE::DaySequence::ComputeSignedDistance(Shape, BlendPosition);
+				return Distance < 0.f ? FMath::Clamp(-Distance / BlendAmount, 0.f, 1.f) : 0.f;
+			};
+			
+			for (const UShapeComponent* Shape : GetVolumeShapeComponents())
+			{
+				InternalBlendWeight = FMath::Max(InternalBlendWeight, GetBlendWeightForShape(Shape));
+			}
+		}
+		else
+		{
+			InternalBlendWeight = 1.f;
+		}
+		break;
 	}
+
+	return InternalBlendWeight;
 }
 
 void UDaySequenceModifierComponent::EmptyVolumeShapeComponents()
@@ -1527,19 +1521,14 @@ void UDaySequenceModifierComponent::InvalidateMuteStates() const
 	OnInvalidateMuteStates.Broadcast();
 }
 
-void UDaySequenceModifierComponent::EnableDistanceVolumeBlends(APlayerController* InActor)
+void UDaySequenceModifierComponent::SetBlendTarget(APlayerController* InActor)
 {	
-	ExternalVolumeBlendTarget = InActor;
-}
-
-void UDaySequenceModifierComponent::SetUseVolume(bool bState)
-{
-	bUseVolume = bState;
+	WeakBlendTarget = InActor;
 }
 
-void UDaySequenceModifierComponent::SetCustomVolumeBlendWeight(float Weight)
+void UDaySequenceModifierComponent::SetUserBlendWeight(float Weight)
 {
-	CustomVolumeBlendWeight = FMath::Clamp(Weight, 0.f, 1.f);
+	UserBlendWeight = FMath::Clamp(Weight, 0.f, 1.f);
 }
 
 #if ENABLE_DRAW_DEBUG
@@ -1566,24 +1555,6 @@ bool UDaySequenceModifierComponent::ShouldShowDebugInfo() const
 }
 #endif
 
-bool UDaySequenceModifierComponent::IsBlendTargetInAnyVolume()
-{
-	OccupiedVolumes = 0;
-
-	if (FVector Position; GetBlendPosition(Position))
-	{
-		for (const UShapeComponent* Shape : GetVolumeShapeComponents())
-		{
-			if (GetDistanceBlendFactorForShape(Shape, Position) > 0.f)
-			{
-				OccupiedVolumes++;
-			}
-		}
-	}
-
-	return OccupiedVolumes > 0;
-}
-
 void UDaySequenceModifierComponent::UpdateCachedExternalShapes() const
 {
 	check (bCachedExternalShapesInvalid)
diff --git a/Engine/Plugins/Experimental/DaySequence/Source/DaySequence/Public/Actors/DaySequenceModifierVolume.h b/Engine/Plugins/Experimental/DaySequence/Source/DaySequence/Public/Actors/DaySequenceModifierVolume.h
--- a/Engine/Plugins/Experimental/DaySequence/Source/DaySequence/Public/Actors/DaySequenceModifierVolume.h
+++ b/Engine/Plugins/Experimental/DaySequence/Source/DaySequence/Public/Actors/DaySequenceModifierVolume.h
@@ -36,24 +36,6 @@ protected:
 
 	
 	
-	/** Volume Setup Functions */
-	
-	/**
-	 * Unconditionally adds shape components to the day sequence modifier via AddShapeComponentsToModifier.
-	 * At runtime this also unconditionally calls SetupVolumeCallbacks.
-	 * In editor this conditionally calls SetupVolumeCallbacks if we are not in an Editor world.
-	 */
-	void VolumeSetup();
-
-	/**
-	 * By default this adds an FComponentReference pointing to DefaultBox to the day sequence modifier.
-	 * Derived classes should override this function to add their own shape components at the proper time.
-	 * If overridden, calling Super::AddShapeComponentsToModifier will determine if DefaultBox is used or ignored.
-	 */
-	virtual void AddShapeComponentsToModifier();
-
-	
-	
 	/** Player Controller Setup Functions */
 	
 	/**
@@ -92,9 +74,6 @@ protected:
 	
 
 	
-	/** Enables this modifier if the current blend target is in the volume. */
-	void TryEnableModifier() const;
-	
 	void SetBlendTarget(APlayerController* InPC);
 	
 protected:
diff --git a/Engine/Plugins/Experimental/DaySequence/Source/DaySequence/Public/DaySequenceModifierComponent.h b/Engine/Plugins/Experimental/DaySequence/Source/DaySequence/Public/DaySequenceModifierComponent.h
--- a/Engine/Plugins/Experimental/DaySequence/Source/DaySequence/Public/DaySequenceModifierComponent.h
+++ b/Engine/Plugins/Experimental/DaySequence/Source/DaySequence/Public/DaySequenceModifierComponent.h
@@ -40,7 +40,7 @@ enum class EDayNightCycleMode : uint8
 	Default,
 	/** Force the day/night cycle to be fixed at the specified constant time */
 	FixedTime,
-	/** Set an initial time for the day/night cycle when this modifier is enabled */
+	/** Set an initial time for the day/night cycle when the modifier is enabled */
 	StartAtSpecifiedTime,
 	/** Use a random, fixed time for the day/night cycle */
 	RandomFixedTime,
@@ -48,14 +48,32 @@ enum class EDayNightCycleMode : uint8
 	RandomStartTime,
 };
 
-/** Enum specifying how to blend into and out of a UDaySequenceModifierComponent that is using a trigger volume */
+/** Enum that defines modifier behavior for auto enabling and computing the internal blend weight. */
 UENUM(BlueprintType)
-enum class EDaySequenceModifierBlendMode : uint8
+enum class EDaySequenceModifierMode : uint8
 {
-	/** Do not blend into the volume - modifier will immediately active when enabled */
-	None,
-	/** Blend the modifier in based on the distance within the volume. Only functional when EnableDistanceVolumeBlends is called. */
-	Distance,
+	// Blend weight is always 1.0.
+	Global,
+
+	// Blend weight smoothly moves between 0.0 and 1.0 as the blend target crosses the volume boundary.
+	Volume
+};
+
+/** Enum specifying how the modifier resolves the user specified blend weight against the internal blend weight. */
+UENUM(BlueprintType)
+enum class EDaySequenceModifierUserBlendPolicy : uint8
+{
+	// User specified weights are ignored (i.e. the effective weight is InternallyComputedWeight
+	Ignored,
+
+	// (default) The effective weight is FMath::Min(InternallyComputedWeight, UserSpecifiedWeight
+	Minimum,
+	
+	// The effective weight is FMath::Max(InternallyComputedWeight, UserSpecifiedWeight
+	Maximum,
+	
+	// The effective weight is UserSpecifiedWeight
+	Override
 };
 
 #if WITH_EDITOR
@@ -146,22 +164,7 @@ public:
 	 */
 	UFUNCTION(BlueprintCallable, Category="Day Sequence")
 	DAYSEQUENCE_API void DisableComponent();
-
-	/**
-	 * Enable this modifier by enabling its sub-sequence within the root Day Sequence
-	 * Will initialize a new sub-sequence track if necessary.
-	 */
-	UFUNCTION(BlueprintCallable, Category="Day Sequence")
-	DAYSEQUENCE_API void EnableModifier();
-
-	/**
-	 * Disable this modifier by disabling its sub-sequence track
-	 */
-	UFUNCTION(BlueprintCallable, Category="Day Sequence")
-	DAYSEQUENCE_API void DisableModifier();
-
-public:
-
+	
 	/**
 	 * Remove all the Sequencer tracks within our procedural Day Sequence
 	 */
@@ -266,51 +269,27 @@ public:
 	UFUNCTION(BlueprintCallable, Category="Day Sequence")
 	DAYSEQUENCE_API void AddVisibilityOverride(UObject* Object, bool bValue);
 
-	/**
-	 * Sets the user day sequence
-	 */
+	/** Sets the user day sequence */
 	UFUNCTION(BlueprintCallable, Category="Day Sequence")
 	DAYSEQUENCE_API void SetUserDaySequence(UDaySequence* InDaySequence);
 
-	/**
-	 * Sets whether to enable modifier based on volume
-	 */
-	DAYSEQUENCE_API void SetUseVolume(bool bState);
-
-	/**
-	 * Enable distance-based volume blends based on the location of the specified actor relative to our volume.
-	 * FinalWeight = Min(DistanceVolumeBlendWeight, CustomVolumeBlendWeight).
-	 */
 	UFUNCTION(BlueprintCallable, Category="Day Sequence")
-	DAYSEQUENCE_API void EnableDistanceVolumeBlends(APlayerController* InActor);
+	DAYSEQUENCE_API void SetMode(EDaySequenceModifierMode NewMode);
 
-	/**
-	 * Determines the number of volumes the current blend target is
-	 * overlapping and returns true if that number is greater than 0
-	 */
-	UFUNCTION(BlueprintCallable, Category="Day Sequence")
-	DAYSEQUENCE_API bool IsBlendTargetInAnyVolume();
-	
-	/**
-	 * Sets a custom blend weight for volume based blends.
-	 * FinalWeight = Min(DistanceVolumeBlendWeight, CustomVolumeBlendWeight).
-	 */
 	UFUNCTION(BlueprintCallable, Category="Day Sequence")
-	DAYSEQUENCE_API void SetCustomVolumeBlendWeight(float Weight);
+	DAYSEQUENCE_API void SetBlendPolicy(EDaySequenceModifierUserBlendPolicy NewPolicy);
 
-	/**
-	 * Get the current blend weight.
-	 */
+	/** Sets the blend target to use when in Volume mode. */
 	UFUNCTION(BlueprintCallable, Category="Day Sequence")
-	float GetCurrentBlendWeight() const;
+	DAYSEQUENCE_API void SetBlendTarget(APlayerController* InActor);
 	
-	float UpdateBlendWeight() const;
+	/** Sets a custom blend weight for volume based blends. Final weight depends on BlendPolicy. */
+	UFUNCTION(BlueprintCallable, Category="Day Sequence")
+	DAYSEQUENCE_API void SetUserBlendWeight(float Weight);
 
-	/**
-	 * Applies the given collision type to all valid volume shape components.
-	 */
+	/** Get the current blend weight. */
 	UFUNCTION(BlueprintCallable, Category="Day Sequence")
-	void SetVolumeCollisionEnabled(const ECollisionEnabled::Type InCollisionType) const;
+	float GetBlendWeight() const;
 
 	void EmptyVolumeShapeComponents();
 	void AddVolumeShapeComponent(const FComponentReference& InShapeReference);
@@ -321,8 +300,6 @@ public:
 	bool ShouldShowDebugInfo() const;
 #endif
 	
-public:
-
 	/*~ Begin UActorComponent interface */
 	void BeginPlay() override;
 	void EndPlay(EEndPlayReason::Type Reason) override;
@@ -348,16 +325,20 @@ public:
 #endif
 
 protected:
+	
+	/** Enable the modifier by enabling its subsection (creating it if necessary) in the Root Sequence. */
+	void EnableModifier();
 
-	virtual bool CanBeEnabled() const;
+	/** Disable the modifier by disabling its subsection. */
+	void DisableModifier();
+	
+	bool CanBeEnabled() const;
 
 	UPROPERTY(Instanced, Transient)
 	TObjectPtr<UDaySequenceModifierEasingFunction> EasingFunction;
 
 	TArray<UShapeComponent*> GetVolumeShapeComponents() const;
-
-protected:
-
+	
 	void SetInitialTimeOfDay();
 
 protected:
@@ -375,7 +356,7 @@ protected:
 
 	/** The actor to use for distance-based volume blend calculations */
 	UPROPERTY(Transient)
-	TWeakObjectPtr<APlayerController> ExternalVolumeBlendTarget;
+	TWeakObjectPtr<APlayerController> WeakBlendTarget;
 
 	/** An optional user-provided Day Sequence - used instead of our procedurally generated one if set */
 	UPROPERTY(BlueprintReadWrite, EditAnywhere, Category="Day Sequence", meta=(EditCondition="!bUseCollection", EditConditionHides, DisplayAfter="bUseCollection"))
@@ -396,24 +377,25 @@ protected:
 	UPROPERTY(BlueprintReadWrite, EditAnywhere, Category="Time", DisplayName="Time", meta=(DisplayAfter="DayNightCycle", EditCondition="DayNightCycle==EDayNightCycleMode::FixedTime || DayNightCycle==EDayNightCycleMode::StartAtSpecifiedTime", EditConditionHides))
 	float DayNightCycleTime;
 
-	/** Defines the blend range based on BlendMode. When BlendMode == Distance, this is interpreted as world units */
-	UPROPERTY(BlueprintReadWrite, EditAnywhere, Category="Day Sequence", meta=(DisplayAfter="BlendMode", EditCondition="BlendMode!=EDaySequenceModifierBlendMode::None", EditConditionHides))
+	/** Defines the region in which the effective blend weight is in the range (0.0, 1.0) (not inclusive) when Mode == EDaySequenceModifierMode::Volume. */
+	UPROPERTY(BlueprintReadWrite, EditAnywhere, Category="Day Sequence", meta=(DisplayAfter="Mode", EditCondition="Mode==EDaySequenceModifierMode::Distance", EditConditionHides))
 	float BlendAmount;
 
-	/** Cached blend factor last time we ticked in the range 0-1. Used to prevent superfluous updates on a paused day sequence. */
-	float CachedBlendFactor;
-
-	/** Defines a custom blend weight. The final blend weight is the minimum between the Distance blend weight and CustomVolumeBlendWeight. */
-	UPROPERTY(Transient)
-	float CustomVolumeBlendWeight;
+	/** User specified blend weight. The final blend weight is determined by BlendPolicy. */
+	UPROPERTY(Transient, BlueprintReadWrite, EditAnywhere, Category="Day Sequence", meta=(DisplayAfter="BlendPolicy", EditCondition="BlendPolicy!=EDaySequenceModifierUserBlendPolicy::Ignored", EditConditionHides))
+	float UserBlendWeight;
 
-	/** Changes the way this modifier controls the day/night cycle time */
+	/** Changes the way the modifier controls the day/night cycle time when enabled. */
 	UPROPERTY(BlueprintReadWrite, EditAnywhere, Category="Time", DisplayName="Day/Night Cycle")
 	EDayNightCycleMode DayNightCycle;
 
-	/** Changes the way this modifier blends in and out */
+	/** Determines how the modifier computes InternalBlendWeight. */
+	UPROPERTY(BlueprintReadWrite, EditAnywhere, Category="Day Sequence")
+	EDaySequenceModifierMode Mode;
+
+	/** Determines how the modifier uses UserBlendWeight to compute effective blend weight. */
 	UPROPERTY(BlueprintReadWrite, EditAnywhere, Category="Day Sequence")
-	EDaySequenceModifierBlendMode BlendMode;
+	EDaySequenceModifierUserBlendPolicy BlendPolicy;
 
 	/** Blueprint exposed delegate invoked after the component's subsequences are reinitialized. */
 	UPROPERTY(BlueprintAssignable, Transient, meta=(AllowPrivateAccess="true"))
@@ -423,23 +405,19 @@ protected:
 	UPROPERTY(BlueprintAssignable, Transient, meta=(AllowPrivateAccess="true"))
 	FOnPostEnableModifier OnPostEnableModifier;
 	
-	/** When enabled, these overrides will always override all settings regardless of their bias */
+	/** When enabled, these overrides will always override all settings regardless of their bias. */
 	UPROPERTY(BlueprintReadWrite, EditAnywhere, Category="Day Sequence", meta=(InlineEditConditionToggle))
 	uint8 bIgnoreBias : 1;
 
-	/** Whether to use the bounding volume for enabling this modifier. If not enabled this modifier will always be active. */
-	UPROPERTY(BlueprintReadWrite, EditAnywhere, Category="Day Sequence")
-	uint8 bUseVolume : 1;
-
 	/** Flag used track whether or not this component is enabled or disabled. */
 	UPROPERTY(Transient, DuplicateTransient, BlueprintReadOnly, Category="Day Sequence")
 	uint8 bIsComponentEnabled : 1;
 	
-	/** Non-serialized variable for tracking whether our overrides are enabled or not */
+	/** Non-serialized variable for tracking whether our overrides are enabled or not. */
 	UPROPERTY(Transient, DuplicateTransient, BlueprintReadOnly, Category="Day Sequence")
 	uint8 bIsEnabled : 1;
 
-	/** Preview this day sequence modifier in the editor */
+	/** When enabled, preview this day sequence modifier in the editor. */
 	UPROPERTY(Transient, BlueprintReadWrite, EditAnywhere, Category="Day Sequence")
 	uint8 bPreview : 1;
 
@@ -455,18 +433,28 @@ protected:
 	uint8 bSmoothBlending : 1;
 
 private:
+	
 	/**
 	 * Get the blend position (handles preview and game world).
 	 * @return Returns true if we have a valid blend position and InPosition was set, false if InPosition was not set.
 	 */
 	bool GetBlendPosition(FVector& InPosition) const;
-	float GetDistanceBlendFactorForShape(const UShapeComponent* Shape, const FVector& Position) const;
 
-	mutable float CachedDistanceBlendFactor;
-	float GetCachedDistanceBlendFactor() const { return CachedDistanceBlendFactor; }
-	float GetDistanceBlendFactor(const FVector& Position) const;
 	/**
-	 * Creates and adds or marks for preserve all subsections that this modifier is responsible for.
+	 * Updates InternalBlendWeight, the update method is determined by Mode.
+	 * @return Returns InternalBlendWeight.
+	 */
+	float UpdateInternalBlendWeight();
+
+	/**
+	 * The blend weight computed by the modifier.
+	 * When this is non-zero the modifier is automatically enabled.
+	 * Used to compute effective blend weight along with BlendPolicy and UserBlendWeight.
+	 */
+	float InternalBlendWeight;
+	
+	/**
+	 * Creates and adds or marks for preserve all subsections that the modifier is responsible for.
 	 * Optionally provided a map of all sections that exist in the root sequence to a bool flag used to mark that section as still relevant.
 	 */
 	void ReinitializeSubSequence(ADaySequenceActor::FSubSectionPreserveMap* SectionsToPreserve);
@@ -478,8 +466,6 @@ private:
 	mutable TArray<TWeakObjectPtr<UShapeComponent>> CachedExternalShapes;
 	mutable bool bCachedExternalShapesInvalid = true;
 	
-	int32 OccupiedVolumes = 0;
-	
 	/*~ Transient state for active gameplay */
 	TWeakObjectPtr<UMovieSceneSubSection> WeakSubSection;
 	TArray<TWeakObjectPtr<UMovieSceneSubSection>> SubSections;
@@ -491,7 +477,7 @@ private:
 	
 	void OnDebugLevelChanged(int32 InDebugLevel);
 	
-	/** Determines whether or not this modifier will show debug info */
+	/** Determines whether or not the modifier will show debug info */
 	int32 DebugLevel;
 
 	TSharedPtr<UE::DaySequence::FDaySequenceDebugEntry> DebugEntry;
