DaySequenceModifierComponent: Fix BlendAmount edit condition

diff --git a/Engine/Plugins/Experimental/DaySequence/Source/DaySequence/Public/DaySequenceModifierComponent.h b/Engine/Plugins/Experimental/DaySequence/Source/DaySequence/Public/DaySequenceModifierComponent.h
--- a/Engine/Plugins/Experimental/DaySequence/Source/DaySequence/Public/DaySequenceModifierComponent.h
+++ b/Engine/Plugins/Experimental/DaySequence/Source/DaySequence/Public/DaySequenceModifierComponent.h
@@ -378,7 +378,7 @@ protected:
 	float DayNightCycleTime;
 
 	/** Defines the region in which the effective blend weight is in the range (0.0, 1.0) (not inclusive) when Mode == EDaySequenceModifierMode::Volume. */
-	UPROPERTY(BlueprintReadWrite, EditAnywhere, Category="Day Sequence", meta=(DisplayAfter="Mode", EditCondition="Mode==EDaySequenceModifierMode::Distance", EditConditionHides))
+	UPROPERTY(BlueprintReadWrite, EditAnywhere, Category="Day Sequence", meta=(DisplayAfter="Mode", EditCondition="Mode==EDaySequenceModifierMode::Volume", EditConditionHides))
 	float BlendAmount;
 
 	/** User specified blend weight. The final blend weight is determined by BlendPolicy. */
