DaySequence: Let static time work without an interp curve.

diff --git a/Engine/Plugins/Experimental/DaySequence/Source/DaySequence/Private/DaySequenceActor.cpp b/Engine/Plugins/Experimental/DaySequence/Source/DaySequence/Private/DaySequenceActor.cpp
--- a/Engine/Plugins/Experimental/DaySequence/Source/DaySequence/Private/DaySequenceActor.cpp
+++ b/Engine/Plugins/Experimental/DaySequence/Source/DaySequence/Private/DaySequenceActor.cpp
@@ -966,7 +966,7 @@ void ADaySequenceActor::InvalidateMuteStates() const
 
 void ADaySequenceActor::WarpEvaluationRange(FMovieSceneEvaluationRange& InOutRange) const
 {
-	if (!bUseInterpCurve || !DayInterpCurve || !RootSequence || bForceDisableDayInterpCurve)
+	if (!RootSequence)
 	{
 		return;
 	}
@@ -1008,6 +1008,7 @@ void ADaySequenceActor::WarpEvaluationRange(FMovieSceneEvaluationRange& InOutRan
 	}
 
 	// Warp with curve
+	if (bUseInterpCurve && DayInterpCurve && !bForceDisableDayInterpCurve)
 	{
 		const float LowerBoundTimeSeconds = static_cast<float>(Range.GetLowerBoundValue() / TickRate);
 		float LowerBoundTimeInHours = DayLengthHours * LowerBoundTimeSeconds / DayCycleSeconds;
