DaySequence: Fix infinite recursion that occurs when the DSA is paused on end play

diff --git a/Engine/Plugins/Experimental/DaySequence/Source/DaySequence/Private/DaySequenceActor.cpp b/Engine/Plugins/Experimental/DaySequence/Source/DaySequence/Private/DaySequenceActor.cpp
--- a/Engine/Plugins/Experimental/DaySequence/Source/DaySequence/Private/DaySequenceActor.cpp
+++ b/Engine/Plugins/Experimental/DaySequence/Source/DaySequence/Private/DaySequenceActor.cpp
@@ -846,8 +846,11 @@ void ADaySequenceActor::OnSequencePlayerUpdate(const UDaySequencePlayer& Player,
 	const float CurrentHours = FrameTimeToDayHours(CurrentTime);
 	const float PreviousHours = FrameTimeToDayHours(PreviousTime);
 	SequencePlayerUpdated(CurrentHours, PreviousHours);
-	
-	OnDaySequenceUpdate.Broadcast();
+
+	if (IsPlaying())
+	{
+		OnDaySequenceUpdate.Broadcast();
+	}
 }
 
 void ADaySequenceActor::SequencePlayerUpdated(float CurrentTime, float PreviousTime)
