DaySequence: Move remainder of procedural subsequence logic out of modifier into procedural day sequence builder

#changelist validated
#virtualized

diff --git a/Engine/Plugins/Experimental/DaySequence/Source/DaySequence/Private/DaySequenceModifierComponent.cpp b/Engine/Plugins/Experimental/DaySequence/Source/DaySequence/Private/DaySequenceModifierComponent.cpp
--- a/Engine/Plugins/Experimental/DaySequence/Source/DaySequence/Private/DaySequenceModifierComponent.cpp
+++ b/Engine/Plugins/Experimental/DaySequence/Source/DaySequence/Private/DaySequenceModifierComponent.cpp
@@ -2,134 +2,29 @@
 
 #include "DaySequenceModifierComponent.h"
 
-#include "Components/BoxComponent.h"
 #include "DaySequence.h"
 #include "DaySequenceCollectionAsset.h"
 #include "DaySequenceModule.h"
 #include "DaySequencePlayer.h"
-#include "DaySequenceTrack.h"
 #include "DaySequenceStaticTime.h"
-
-#include "Engine/World.h"
-#include "Materials/MaterialInterface.h"
-
-#include "Components/CapsuleComponent.h"
-#include "Components/SphereComponent.h"
+#include "DaySequenceTrack.h"
 
 #include "MovieScene.h"
-#include "MovieSceneTrack.h"
 #include "MovieSceneSection.h"
-#include "MovieSceneBinding.h"
-#include "MovieScenePossessable.h"
-
-#include "Tracks/MovieSceneSubTrack.h"
-#include "Tracks/MovieSceneBoolTrack.h"
-#include "Tracks/MovieSceneFloatTrack.h"
-#include "Tracks/MovieSceneColorTrack.h"
-#include "Tracks/MovieSceneDoubleTrack.h"
-#include "Tracks/MovieSceneVectorTrack.h"
-#include "Tracks/MovieSceneMaterialTrack.h"
-#include "Tracks/MovieSceneVisibilityTrack.h"
-#include "Tracks/MovieScene3DTransformTrack.h"
-#include "Tracks/MovieScenePrimitiveMaterialTrack.h"
-
-#include "Sections/MovieSceneSubSection.h"
-#include "Sections/MovieSceneBoolSection.h"
-#include "Sections/MovieSceneFloatSection.h"
-#include "Sections/MovieSceneColorSection.h"
-#include "Sections/MovieSceneDoubleSection.h"
-#include "Sections/MovieSceneVectorSection.h"
-#include "Sections/MovieScene3DTransformSection.h"
-#include "Sections/MovieScenePrimitiveMaterialSection.h"
-
+#include "MovieSceneTrack.h"
+#include "Components/BoxComponent.h"
+#include "Components/CapsuleComponent.h"
+#include "Components/SphereComponent.h"
+#include "Engine/World.h"
+#include "Materials/MaterialInterface.h"
 #include "ProfilingDebugging/CsvProfiler.h"
-
-#if ENABLE_DRAW_DEBUG
-#endif
+#include "Sections/MovieSceneSubSection.h"
+#include "Tracks/MovieSceneSubTrack.h"
 
 #define LOCTEXT_NAMESPACE "DaySequenceModifierComponent"
 
 namespace UE::DaySequence
 {
-	template<typename TrackType>
-	TrackType* CreateOrAddOverrideTrack(UMovieScene* MovieScene, const FGuid& ObjectGuid, FName Name = NAME_None)
-	{
-		TrackType* Track = MovieScene->FindTrack<TrackType>(ObjectGuid, Name);
-		if (!Track)
-		{
-			// Clear RF_Transactional and set RF_Transient on created tracks and sections
-			// to avoid dirtying the package for these procedurally generated sequences.
-			// RF_Transactional is explicitly set in UMovieSceneSection/Track::PostInitProperties.
-			Track = NewObject<TrackType>(MovieScene, NAME_None, RF_Transient);
-			Track->ClearFlags(RF_Transactional);
-
-			UMovieSceneSection* Section = Track->CreateNewSection();
-			Section->ClearFlags(RF_Transactional);
-			Section->SetFlags(RF_Transient);
-			Section->SetRange(TRange<FFrameNumber>::All());
-
-			Track->AddSection(*Section);
-			MovieScene->AddGivenTrack(Track, ObjectGuid);
-		}
-
-		return Track;
-	}
-
-	template<typename TrackType>
-	TrackType* CreateOrAddPropertyOverrideTrack(UMovieScene* MovieScene, const FGuid& ObjectGuid, FName InPropertyName)
-	{
-		TrackType* Track = CreateOrAddOverrideTrack<TrackType>(MovieScene, ObjectGuid, InPropertyName);
-		check(Track);
-		
-		const FString PropertyPath = InPropertyName.ToString();
-
-		// Split the property path to capture the leaf property name and parent struct to conform
-		// with Sequencer Editor property name/path and display name conventions:
-		//
-		// PropertyName = MyProperty
-		// PropertyPath = MyPropertyStruct.MyProperty
-		// DisplayName = PropertyName (PropertyStruct)
-		FName PropertyName;
-		FName PropertyParent;
-		int32 NamePos = INDEX_NONE;
-		if (PropertyPath.FindLastChar('.', NamePos) && NamePos < PropertyPath.Len() - 1)
-		{
-			PropertyName = FName(FStringView(*PropertyPath + NamePos + 1, PropertyPath.Len() - NamePos - 1));
-			PropertyParent = FName(FStringView(*PropertyPath, NamePos));
-		}
-		else
-		{
-			PropertyName = *PropertyPath;
-		}
-		
-		Track->SetPropertyNameAndPath(PropertyName, PropertyPath);
-
-#if WITH_EDITOR
-		if (NamePos != INDEX_NONE)
-		{
-			FText DisplayText = FText::Format(LOCTEXT("ModifierPropertyTrackFormat", "{0} ({1})"), FText::FromName(PropertyName), FText::FromName(PropertyParent));
-			Track->SetDisplayName(DisplayText);
-		}
-#endif
-		return Track;
-	}
-
-	template<typename TrackType, typename SectionType>
-	SectionType* CreateOrAddPropertyOverrideSection(UMovieScene* MovieScene, const FGuid& ObjectGuid, FName PropertyName)
-	{
-		TrackType* Track = CreateOrAddPropertyOverrideTrack<TrackType>(MovieScene, ObjectGuid, PropertyName);
-		check(Track);
-		return Cast<SectionType>(Track->GetAllSections()[0]);
-	}
-
-	template<typename TrackType, typename SectionType>
-	SectionType* CreateOrAddOverrideSection(UMovieScene* MovieScene, const FGuid& ObjectGuid)
-	{
-		TrackType* Track = CreateOrAddOverrideTrack<TrackType>(MovieScene, ObjectGuid);
-		check(Track);
-		return Cast<SectionType>(Track->GetAllSections()[0]);
-	}
-
 	FVector GVolumePreviewLocation = FVector::ZeroVector;
 	bool bIsSimulating = false;
 
@@ -190,24 +85,6 @@ namespace UE::DaySequence
 		// @todo: unsupported shape?
 		return (InWorldPosition - ShapeComponent->GetComponentLocation()).Length();
 	}
-
-	bool TestValidProperty(UObject* Object, FProperty* Property)
-	{
-		if (!Property)
-		{
-			FFrame::KismetExecutionMessage(*FString::Printf(TEXT("Invalid property specified for object %s."), *Object->GetName()), ELogVerbosity::Error);
-			return false;
-		}
-
-		if (Property->HasAnyPropertyFlags(CPF_Deprecated))
-		{
-			// Emit a warning for deprecated properties but still consider them valid
-			FFrame::KismetExecutionMessage(*FString::Printf(TEXT("Depcrecated property specified: %s for object %s."), *Property->GetName(), *Object->GetName()), ELogVerbosity::Warning);
-		}
-
-		return true;
-	}
-
 } // namespace UE::DaySequence
 
 void UDaySequenceModifierEasingFunction::Initialize(EEasingFunctionType EasingType)
@@ -317,11 +194,6 @@ void UDaySequenceModifierComponent::UpdateEditorPreview(float DeltaTime)
 				// If we're using a blend we have to mark active sections as changed
 				// in order to force an update in-editor:
 				
-				if (UMovieSceneSubSection* SubSection = WeakSubSection.Get())
-				{
-					SubSection->MarkAsChanged();
-				}
-
 				for (TWeakObjectPtr<UMovieSceneSubSection> SubSection : SubSections)
 				{
 					UMovieSceneSubSection* StrongSubSection = SubSection.Get();
@@ -444,29 +316,6 @@ void UDaySequenceModifierComponent::EndPlay(EEndPlayReason::Type Reason)
 	RemoveSubSequenceTrack();
 }
 
-void UDaySequenceModifierComponent::ResetOverrides()
-{
-	if (ProceduralDaySequence)
-	{
-		UMovieScene* MovieScene = ProceduralDaySequence->GetMovieScene();
-		TArray<FMovieSceneBinding> Bindings = MovieScene->GetBindings();
-
-		for (const FMovieSceneBinding& Binding : Bindings)
-		{
-			FGuid BindingID = Binding.GetObjectGuid();
-
-			ProceduralDaySequence->UnbindPossessableObjects(BindingID);
-			MovieScene->RemovePossessable(BindingID);
-		}
-	}
-
-	if (TargetActor && bUnpauseOnDisable)
-	{
-		TargetActor->Play();
-		bUnpauseOnDisable = false;
-	}
-}
-
 void UDaySequenceModifierComponent::BindToDaySequenceActor(ADaySequenceActor* DaySequenceActor)
 {
 	if (TargetActor == DaySequenceActor)
@@ -543,9 +392,6 @@ void UDaySequenceModifierComponent::RemoveSubSequenceTrack()
 		}
 	};
 	
-	RemoveSubTrack(WeakSubSection.Get());
-	WeakSubSection = nullptr;
-
 	for (const TWeakObjectPtr<UMovieSceneSubSection> SubSection : SubSections)
 	{
 		RemoveSubTrack(SubSection.Get());
@@ -795,11 +641,7 @@ void UDaySequenceModifierComponent::ReinitializeSubSequence(ADaySequenceActor::F
 		}
 		else
 		{
-			// Always create the sub section even if it is nullptr. This means that
-			// if the procedural day sequence is created later, we can still add it to the sub section
-			UDaySequence* SequenceToUse = UserDaySequence ? UserDaySequence : ProceduralDaySequence;
-	
-			WeakSubSection = InitializeDaySequence(SequenceToUse);
+			InitializeDaySequence(UserDaySequence ? UserDaySequence : TransientSequence);
 		}
 #if ROOT_SEQUENCE_RECONSTRUCTION_ENABLED
 	}
@@ -962,407 +804,48 @@ UMovieSceneSubSection* UDaySequenceModifierComponent::InitializeDaySequence(cons
 	return SubSection;
 }
 
-FGuid UDaySequenceModifierComponent::GetOrCreateProceduralBinding(UObject* Object)
-{
-	if (!Object)
-	{
-		FFrame::KismetExecutionMessage(TEXT("Null Object parameter specified."), ELogVerbosity::Error);
-		return FGuid();
-	}
-
-	USceneComponent* Component = Cast<USceneComponent>(Object);
-	AActor*          Actor     = Cast<AActor>(Object);
-
-	// Set up the time of day actor binding if we don't have one already
-	if (!TargetActor)
-	{
-		if (Actor)
-		{
-			TargetActor = Cast<ADaySequenceActor>(Actor);
-		}
-		else if (Component)
-		{
-			TargetActor = Cast<ADaySequenceActor>(Component->GetOwner());
-		}
-	}
-
-	if (!TargetActor)
-	{
-		FFrame::KismetExecutionMessage(TEXT("No valid ADaySequenceActor has been set up. Have you called BindToDaySequenceActor yet?"), ELogVerbosity::Error);
-		return FGuid();
-	}
-
-	if (Component && !Component->IsIn(TargetActor))
-	{
-		FFrame::KismetExecutionMessage(TEXT("Unable to bind to components that exist outside of the ADaySequenceActor we are tracking."), ELogVerbosity::Error);
-		return FGuid();
-	}
-
-	if (Actor && Actor != TargetActor)
-	{
-		FFrame::KismetExecutionMessage(TEXT("Unable to bind to actors that are not the ADaySequenceActor we are tracking."), ELogVerbosity::Error);
-		return FGuid();
-	}
-
-	if (!ProceduralDaySequence)
-	{
-		// Name the procedural sequence the same as this component's owner so it shows up in the Sequence
-		// with a meaningful name
-#if WITH_EDITOR
-		FName SequenceName = MakeUniqueObjectName(this, UDaySequence::StaticClass(), *GetOwner()->GetActorLabel());
-#else
-		FName SequenceName = MakeUniqueObjectName(this, UDaySequence::StaticClass(), GetOwner()->GetFName());
-#endif
-
-		ProceduralDaySequence = NewObject<UDaySequence>(this, SequenceName);
-		ProceduralDaySequence->Initialize(RF_Transient);
-	}
-
-	// If we have a sub section but it has no sequence applied, apply it now.
-	// This implies EnableModifier was called before we had any valid sequence data
-	UMovieSceneSubSection* SubSection = WeakSubSection.Get();
-	if (SubSection && SubSection->GetSequence() == nullptr)
-	{
-		SubSection->SetIsLocked(false);
-		SubSection->MarkAsChanged();
-		SubSection->SetSequence(ProceduralDaySequence);
-		SubSection->SetIsLocked(true);
-	}
-
-	UMovieScene* MovieScene = ProceduralDaySequence->GetMovieScene();
-
-	TSharedRef<const UE::MovieScene::FSharedPlaybackState> SharedPlaybackState = MovieSceneHelpers::CreateTransientSharedPlaybackState(TargetActor, ProceduralDaySequence);
-
-	// Find the main binding
-	FGuid RootGuid = ProceduralDaySequence->FindBindingFromObject(TargetActor, SharedPlaybackState);
-	if (!RootGuid.IsValid())
-	{
-		FString RootName = TargetActor->GetName();
-		FMovieScenePossessable Possessable(RootName, TargetActor->GetClass());
-		FMovieSceneBinding     Binding(Possessable.GetGuid(), RootName);
-
-		RootGuid = Possessable.GetGuid();
-
-		// Explicitly invoke MarkAsChanged to ensure proper notification at runtime.
-		// The Modify that AddPossessable invokes only works in editor.
-		MovieScene->MarkAsChanged();
-		MovieScene->AddPossessable(Possessable, Binding);
-		ProceduralDaySequence->BindPossessableObject(RootGuid, *TargetActor, TargetActor);
-	}
-
-	// If we're trying to animate the actor, just return the root binding
-	if (Actor)
-	{
-		return RootGuid;
-	}
-
-	// If we're trying to animate a component within the actor, retrieve or create a child binding for that
-	FGuid ComponentGuid = ProceduralDaySequence->FindBindingFromObject(Component, SharedPlaybackState);
-	if (!ComponentGuid.IsValid() && Component)
-	{
-		FString Name = Component->GetName();
-		FMovieScenePossessable Possessable(Name, Component->GetClass());
-		FMovieSceneBinding     Binding(Possessable.GetGuid(), Name);
-
-		Possessable.SetParent(RootGuid, MovieScene);
-		ComponentGuid = Possessable.GetGuid();
-
-		// Explicitly invoke MarkAsChanged to ensure proper notification at runtime.
-		// The Modify that AddPossessable invokes only works in editor.
-		MovieScene->MarkAsChanged();
-		MovieScene->AddPossessable(Possessable, Binding);
-		ProceduralDaySequence->BindPossessableObject(ComponentGuid, *Component, TargetActor);
-	}
-
-	return ComponentGuid;
-}
-
-void UDaySequenceModifierComponent::AddScalarOverride(UObject* Object, FName PropertyName, double Value)
-{
-	using namespace UE::DaySequence;
-	using namespace	UE::MovieScene;
-
-	FGuid ObjectGuid = GetOrCreateProceduralBinding(Object);
-	if (!ObjectGuid.IsValid())
-	{
-		return;
-	}
-
-	check(ProceduralDaySequence);
-
-	// Create the new track
-	UMovieScene* MovieScene = ProceduralDaySequence->GetMovieScene();
-
-	FTrackInstancePropertyBindings Bindings(PropertyName, PropertyName.ToString());
-	FProperty* Property = Bindings.GetProperty(*Object);
-	if (!TestValidProperty(Object, Property))
-	{
-		// Do nothing
-	}
-	else if (Property->IsA<FFloatProperty>())
-	{
-		UMovieSceneFloatSection* FloatSection = CreateOrAddPropertyOverrideSection<UMovieSceneFloatTrack, UMovieSceneFloatSection>(MovieScene, ObjectGuid, PropertyName);
-		FloatSection->GetChannel().SetDefault(static_cast<float>(Value));
-	}
-	else if (Property->IsA<FDoubleProperty>())
-	{
-		UMovieSceneDoubleSection* DoubleSection = CreateOrAddPropertyOverrideSection<UMovieSceneDoubleTrack, UMovieSceneDoubleSection>(MovieScene, ObjectGuid, PropertyName);
-		DoubleSection->GetChannel().SetDefault(Value);
-	}
-	else
-	{
-		FFrame::KismetExecutionMessage(*FString::Printf(TEXT("Unable to animate a %s property as a scalar."), *Property->GetClass()->GetName()), ELogVerbosity::Error);
-	}
-}
-
-void UDaySequenceModifierComponent::AddColorOverride(UObject* Object, FName PropertyName, FLinearColor Value)
-{
-	using namespace UE::DaySequence;
-
-	FGuid ObjectGuid = GetOrCreateProceduralBinding(Object);
-	if (!ObjectGuid.IsValid())
-	{
-		return;
-	}
-
-	check(ProceduralDaySequence);
-
-	// Create the new track
-	UMovieScene* MovieScene = ProceduralDaySequence->GetMovieScene();
-
-	FTrackInstancePropertyBindings Bindings(PropertyName, PropertyName.ToString());
-	FProperty* Property = Bindings.GetProperty(*Object);
-	if (!TestValidProperty(Object, Property))
-	{
-		// Do nothing
-	}
-	else if (Property->IsA<FStructProperty>() && 
-		(CastField<FStructProperty>(Property)->Struct == TBaseStructure<FLinearColor>::Get() || CastField<FStructProperty>(Property)->Struct == TBaseStructure<FColor>::Get()) )
-	{
-		UMovieSceneColorSection* ColorSection = CreateOrAddPropertyOverrideSection<UMovieSceneColorTrack, UMovieSceneColorSection>(MovieScene, ObjectGuid, PropertyName);
-
-		ColorSection->GetRedChannel().SetDefault(Value.R);
-		ColorSection->GetGreenChannel().SetDefault(Value.G);
-		ColorSection->GetBlueChannel().SetDefault(Value.B);
-		ColorSection->GetAlphaChannel().SetDefault(Value.A);
-	}
-}
-
-void UDaySequenceModifierComponent::AddMaterialOverride(UObject* Object, int32 MaterialIndex, UMaterialInterface* Value)
-{
-	using namespace UE::DaySequence;
-
-	FGuid ObjectGuid = GetOrCreateProceduralBinding(Object);
-	if (!ObjectGuid.IsValid())
-	{
-		return;
-	}
-
-	check(ProceduralDaySequence);
-
-	// Create the new track
-	UMovieScene* MovieScene = ProceduralDaySequence->GetMovieScene();
-
-	UMovieScenePrimitiveMaterialTrack* MaterialTrack = CreateOrAddOverrideTrack<UMovieScenePrimitiveMaterialTrack>(MovieScene, ObjectGuid);
-	MaterialTrack->SetMaterialInfo(FComponentMaterialInfo{ FName(), MaterialIndex, EComponentMaterialType::IndexedMaterial });
-
-	UMovieScenePrimitiveMaterialSection* Section = Cast<UMovieScenePrimitiveMaterialSection>(MaterialTrack->GetAllSections()[0]);
-	Section->MaterialChannel.SetDefault(Value);
-}
-
-void UDaySequenceModifierComponent::AddScalarMaterialParameterOverride(UObject* Object, int32 MaterialIndex, FName ParameterName, float Value)
-{
-	using namespace UE::DaySequence;
-
-	FGuid ObjectGuid = GetOrCreateProceduralBinding(Object);
-	if (!ObjectGuid.IsValid())
-	{
-		return;
-	}
-
-	check(ProceduralDaySequence);
-
-	// Create the new track or locate an existing one
-	UMovieScene* MovieScene = ProceduralDaySequence->GetMovieScene();
-
-	// Material parameter tracks use the material index as the unique name
-	FName IndexAsName(*FString::FromInt(MaterialIndex));
-	UMovieSceneComponentMaterialTrack* MaterialTrack = CreateOrAddOverrideTrack<UMovieSceneComponentMaterialTrack>(MovieScene, ObjectGuid, IndexAsName);
-
-	MaterialTrack->SetMaterialInfo(FComponentMaterialInfo{FName(), MaterialIndex, EComponentMaterialType::IndexedMaterial });
-	MaterialTrack->AddScalarParameterKey(ParameterName, 0, Value);
-}
-
-void UDaySequenceModifierComponent::AddColorMaterialParameterOverride(UObject* Object, int32 MaterialIndex, FName ParameterName, FLinearColor Value)
-{
-	using namespace UE::DaySequence;
-
-	FGuid ObjectGuid = GetOrCreateProceduralBinding(Object);
-	if (!ObjectGuid.IsValid())
-	{
-		return;
-	}
-
-	check(ProceduralDaySequence);
-
-	// Create the new track or locate an existing one
-	UMovieScene* MovieScene = ProceduralDaySequence->GetMovieScene();
-
-	// Material parameter tracks use the material index as the unique name
-	FName IndexAsName(*FString::FromInt(MaterialIndex));
-	UMovieSceneComponentMaterialTrack* MaterialTrack = CreateOrAddOverrideTrack<UMovieSceneComponentMaterialTrack>(MovieScene, ObjectGuid, IndexAsName);
-
-	MaterialTrack->SetMaterialInfo(FComponentMaterialInfo{ FName(), MaterialIndex, EComponentMaterialType::IndexedMaterial });
-	MaterialTrack->AddColorParameterKey(ParameterName, 0, Value);
-}
-
-void UDaySequenceModifierComponent::AddVectorOverride(UObject* Object, FName PropertyName, FVector Value)
+void UDaySequenceModifierComponent::SetUserDaySequence(UDaySequence* InDaySequence)
 {
-	using namespace UE::DaySequence;
-
-	FGuid ObjectGuid = GetOrCreateProceduralBinding(Object);
-	if (!ObjectGuid.IsValid())
+	if (InDaySequence && InDaySequence->HasAnyFlags(RF_Transient))
 	{
+		FFrame::KismetExecutionMessage(TEXT("SetUserDaySequence called with a transient sequence, use SetTransientSequence instead!"), ELogVerbosity::Error);
 		return;
 	}
 
-	check(ProceduralDaySequence);
-
-	// Create the new track
-	UMovieScene* MovieScene = ProceduralDaySequence->GetMovieScene();
-
-	FTrackInstancePropertyBindings Bindings(PropertyName, PropertyName.ToString());
-	FProperty* Property = Bindings.GetProperty(*Object);
-	if (!TestValidProperty(Object, Property))
-	{
-		// Do nothing
-	}
-	else if (Property->IsA<FStructProperty>() && CastField<FStructProperty>(Property)->Struct == TBaseStructure<FVector>::Get())
-	{
-		UMovieSceneDoubleVectorSection* VectorSection = CreateOrAddPropertyOverrideSection<UMovieSceneDoubleVectorTrack, UMovieSceneDoubleVectorSection>(MovieScene, ObjectGuid, PropertyName);
-
-		VectorSection->SetChannelsUsed(3);
-
-		VectorSection->GetChannelProxy().GetChannel<FMovieSceneDoubleChannel>(0)->SetDefault(Value.X);
-		VectorSection->GetChannelProxy().GetChannel<FMovieSceneDoubleChannel>(1)->SetDefault(Value.Y);
-		VectorSection->GetChannelProxy().GetChannel<FMovieSceneDoubleChannel>(2)->SetDefault(Value.Z);
-	}
+	UserDaySequence = InDaySequence;
+	ReinitializeSubSequence(nullptr);
 }
 
-void UDaySequenceModifierComponent::AddTransformOverride(UObject* Object, FTransform Value)
+UDaySequence* UDaySequenceModifierComponent::GetUserDaySequence()
 {
-	using namespace UE::DaySequence;
-
-	FGuid ObjectGuid = GetOrCreateProceduralBinding(Object);
-	if (!ObjectGuid.IsValid())
-	{
-		return;
-	}
-
-	check(ProceduralDaySequence);
-
-	// Create the new track
-	UMovieScene* MovieScene = ProceduralDaySequence->GetMovieScene();
-
-	UMovieScene3DTransformSection* TransformSection = CreateOrAddPropertyOverrideSection<UMovieScene3DTransformTrack, UMovieScene3DTransformSection>(MovieScene, ObjectGuid, "Transform");
-
-	TransformSection->GetChannelProxy().GetChannel<FMovieSceneDoubleChannel>(0)->SetDefault(Value.GetLocation().X);
-	TransformSection->GetChannelProxy().GetChannel<FMovieSceneDoubleChannel>(1)->SetDefault(Value.GetLocation().Y);
-	TransformSection->GetChannelProxy().GetChannel<FMovieSceneDoubleChannel>(2)->SetDefault(Value.GetLocation().Z);
-
-	TransformSection->GetChannelProxy().GetChannel<FMovieSceneDoubleChannel>(3)->SetDefault(Value.Rotator().Roll);
-	TransformSection->GetChannelProxy().GetChannel<FMovieSceneDoubleChannel>(4)->SetDefault(Value.Rotator().Pitch);
-	TransformSection->GetChannelProxy().GetChannel<FMovieSceneDoubleChannel>(5)->SetDefault(Value.Rotator().Yaw);
-
-	TransformSection->GetChannelProxy().GetChannel<FMovieSceneDoubleChannel>(6)->SetDefault(Value.GetScale3D().X);
-	TransformSection->GetChannelProxy().GetChannel<FMovieSceneDoubleChannel>(7)->SetDefault(Value.GetScale3D().Y);
-	TransformSection->GetChannelProxy().GetChannel<FMovieSceneDoubleChannel>(8)->SetDefault(Value.GetScale3D().Z);
+	return UserDaySequence;
 }
 
-void UDaySequenceModifierComponent::AddStaticTimeOfDayOverride(ADaySequenceActor* Actor, float Hours)
+void UDaySequenceModifierComponent::SetTransientSequence(UDaySequence* InDaySequence)
 {
-	using namespace UE::DaySequence;
-
-	FGuid ObjectGuid = GetOrCreateProceduralBinding(Actor);
-	if (!ObjectGuid.IsValid())
+	if (InDaySequence && !InDaySequence->HasAnyFlags(RF_Transient))
 	{
+		FFrame::KismetExecutionMessage(TEXT("SetTransientSequence called with a non-transient sequence, use SetUserDaySequence instead!"), ELogVerbosity::Error);
 		return;
 	}
-
-	check(ProceduralDaySequence);
-
-	// Create the new track
-	UMovieScene* MovieScene = ProceduralDaySequence->GetMovieScene();
-
-	UMovieSceneFloatSection* Section = CreateOrAddPropertyOverrideSection<UMovieSceneFloatTrack, UMovieSceneFloatSection>(MovieScene, ObjectGuid, "StaticTimeOfDay");
-	Section->GetChannel().SetDefault(Hours);
+	
+	TransientSequence = InDaySequence;
+	ReinitializeSubSequence(nullptr);
 }
-
-void UDaySequenceModifierComponent::AddBoolOverride(UObject* Object, FName PropertyName, bool bValue)
+	
+UDaySequence* UDaySequenceModifierComponent::GetTransientSequence()
 {
-	using namespace UE::DaySequence;
-
-	FGuid ObjectGuid = GetOrCreateProceduralBinding(Object);
-	if (!ObjectGuid.IsValid())
-	{
-		return;
-	}
-
-	check(ProceduralDaySequence);
-
-	// Create the new track
-	UMovieScene* MovieScene = ProceduralDaySequence->GetMovieScene();
-
-	FTrackInstancePropertyBindings Bindings(PropertyName, PropertyName.ToString());
-	FProperty* Property = Bindings.GetProperty(*Object);
-	if (!TestValidProperty(Object, Property))
-	{
-		// Do nothing
-	}
-	else if (Property->IsA<FBoolProperty>())
-	{
-		UMovieSceneBoolSection* Section = CreateOrAddPropertyOverrideSection<UMovieSceneBoolTrack, UMovieSceneBoolSection>(MovieScene, ObjectGuid, PropertyName);
-		Section->GetChannel().SetDefault(bValue);
-	}
+	return TransientSequence;
 }
 
-void UDaySequenceModifierComponent::AddVisibilityOverride(UObject* Object, bool bValue)
+void UDaySequenceModifierComponent::SetDayNightCycle(EDayNightCycleMode NewMode)
 {
-	using namespace UE::DaySequence;
-
-	FGuid ObjectGuid = GetOrCreateProceduralBinding(Object);
-	if (!ObjectGuid.IsValid())
-	{
-		return;
-	}
-
-	check(ProceduralDaySequence);
-
-	// Create the new track
-	UMovieScene* MovieScene = ProceduralDaySequence->GetMovieScene();
-
-	static const FName ActorVisibilityTrackName = TEXT("bHidden");
-	static const FName ComponentVisibilityTrackName = TEXT("bHiddenInGame");
-
-	const bool bIsComponent = Object->IsA<USceneComponent>();
-	const bool bIsActor     = Object->IsA<AActor>();
-
-	if (bIsComponent)
-	{
-		UMovieSceneBoolSection* VisibilitySection = CreateOrAddPropertyOverrideSection<UMovieSceneVisibilityTrack, UMovieSceneBoolSection>(MovieScene, ObjectGuid, "bHiddenInGame");
-		VisibilitySection->GetChannel().SetDefault(bValue);
-	}
-	else if (bIsActor)
-	{
-		UMovieSceneBoolSection* VisibilitySection = CreateOrAddPropertyOverrideSection<UMovieSceneVisibilityTrack, UMovieSceneBoolSection>(MovieScene, ObjectGuid, "bHidden");
-		VisibilitySection->GetChannel().SetDefault(bValue);
-	}
+	DayNightCycle = NewMode;
 }
 
-void UDaySequenceModifierComponent::SetUserDaySequence(UDaySequence* InDaySequence)
+void UDaySequenceModifierComponent::SetDayNightCycleTime(float Time)
 {
-	UserDaySequence = InDaySequence;
-	ReinitializeSubSequence(nullptr);
+	DayNightCycleTime = Time;
 }
 
 void UDaySequenceModifierComponent::SetMode(EDaySequenceModifierMode NewMode)
diff --git a/Engine/Plugins/Experimental/DaySequence/Source/DaySequence/Private/ProceduralDaySequenceBuilder.cpp b/Engine/Plugins/Experimental/DaySequence/Source/DaySequence/Private/ProceduralDaySequenceBuilder.cpp
--- a/Engine/Plugins/Experimental/DaySequence/Source/DaySequence/Private/ProceduralDaySequenceBuilder.cpp
+++ b/Engine/Plugins/Experimental/DaySequence/Source/DaySequence/Private/ProceduralDaySequenceBuilder.cpp
@@ -2,21 +2,30 @@
 
 #include "ProceduralDaySequenceBuilder.h"
 
+#include "DaySequence.h"
 #include "DaySequenceActor.h"
 #include "DaySequenceTime.h"
-#include "MovieSceneCommonHelpers.h"
 
 #include "Sections/MovieSceneBoolSection.h"
 #include "Sections/MovieSceneDoubleSection.h"
 #include "Sections/MovieSceneFloatSection.h"
+#include "Sections/MovieSceneVectorSection.h"
+#include "Sections/MovieSceneColorSection.h"
+#include "Sections/MovieScenePrimitiveMaterialSection.h"
 
 #include "Tracks/MovieSceneBoolTrack.h"
 #include "Tracks/MovieSceneDoubleTrack.h"
 #include "Tracks/MovieSceneFloatTrack.h"
+#include "Tracks/MovieSceneVectorTrack.h"
 #include "Tracks/MovieScene3DTransformTrack.h"
+#include "Tracks/MovieSceneColorTrack.h"
+#include "Tracks/MovieScenePrimitiveMaterialTrack.h"
+#include "Tracks/MovieSceneVisibilityTrack.h"
 
+#include "Materials/MaterialInterface.h"
 #include "MovieScene.h"
-#include "DaySequence.h"
+#include "MovieSceneCommonHelpers.h"
+#include "Sections/MovieSceneVisibilitySection.h"
 
 #define LOCTEXT_NAMESPACE "ProceduralDaySequenceBuilder"
 
@@ -102,7 +111,7 @@
 
 		if (bClearInitialSequence)
 		{
-			ClearAllKeys();
+			ClearKeys();
 		}
 	}
 	else
@@ -156,32 +165,67 @@
 	ActiveBinding = GetOrCreateProceduralBinding(InObject);
 }
 
-void UProceduralDaySequenceBuilder::ClearAllKeys()
+void UProceduralDaySequenceBuilder::AddBoolOverride(FName PropertyName, bool Value)
 {
-	if (!ProceduralDaySequence)
+	TPair<float, bool> A = {0.f, Value};
+	TPair<float, bool> B = {1.f, Value};
+
+	AddBoolKeys(PropertyName, {A, B});
+}
+
+void UProceduralDaySequenceBuilder::AddBoolKey(FName PropertyName, float Key, bool Value)
+{
+	AddBoolKey(PropertyName, TPair<float, bool>(Key, Value));
+}
+
+void UProceduralDaySequenceBuilder::AddBoolKey(FName PropertyName, const TPair<float, bool>& KeyValue)
+{
+	AddBoolKeys(PropertyName, TArray {KeyValue});
+}
+
+void UProceduralDaySequenceBuilder::AddBoolKeys(FName PropertyName, const TArray<TPair<float, bool>>& KeysAndValues)
+{
+	if (!IsInitialized())
 	{
+		FFrame::KismetExecutionMessage(TEXT("AddBoolKey(s) called on an uninitialized Procedural Day Sequence Builder!"), ELogVerbosity::Error);
 		return;
 	}
 	
-	if (UMovieScene* MovieScene = ProceduralDaySequence->GetMovieScene())
+	const UMovieScene* MovieScene = ProceduralDaySequence->GetMovieScene();
+	
+	FTrackInstancePropertyBindings Bindings(PropertyName, PropertyName.ToString());
+	FProperty* Property = Bindings.GetProperty(*ActiveBoundObject);
+	if (!UE::DaySequence::IsPropertyValid(ActiveBoundObject, Property))
 	{
-		for (const FMovieSceneBinding& Binding : MovieScene->GetBindings())
+		// Do nothing
+	}
+	else if (Property->IsA<FBoolProperty>())
+	{
+		UMovieSceneBoolSection* Section = CreateOrAddPropertyOverrideSection<UMovieSceneBoolTrack, UMovieSceneBoolSection>(PropertyName);
+		
+		for (const TPair<float, bool>& KeyValue : KeysAndValues)
 		{
-			// Inconvenient we have to do this but at least FindBinding is doing a binary search and we do this once per binding.
-			if (FMovieSceneBinding* MutableBinding = MovieScene->FindBinding(Binding.GetObjectGuid()))
-			{
-				// We have to copy the array here because we are mutating the internal array
-				for (TArray<UMovieSceneTrack*> Tracks = MutableBinding->GetTracks(); UMovieSceneTrack* Track : Tracks)
-				{
-					MutableBinding->RemoveTrack(*Track, MovieScene);
-				}
-			}
+			const FFrameNumber FrameNumber = UE::DaySequence::GetKeyFrameNumber(KeyValue.Key, MovieScene->GetPlaybackRange());
+			
+			Section->GetChannel().AddKeys(TArray {FrameNumber}, TArray {KeyValue.Value});
 		}
 
-		MovieScene->MarkAsChanged();
+		Section->MarkAsChanged();
+	}
+	else
+	{
+		FFrame::KismetExecutionMessage(*FString::Printf(TEXT("Unable to animate a %s property as a bool."), *Property->GetClass()->GetName()), ELogVerbosity::Error);
 	}
 }
 
+void UProceduralDaySequenceBuilder::AddScalarOverride(FName PropertyName, double Value)
+{
+	TPair<float, double> A = {0.f, Value};
+	TPair<float, double> B = {1.f, Value};
+
+	AddScalarKeys(PropertyName, {A, B}, RCIM_Linear);
+}
+
 void UProceduralDaySequenceBuilder::AddScalarKey(FName PropertyName, float Key, double Value, ERichCurveInterpMode InterpMode)
 {
 	AddScalarKey(PropertyName, TPair<float, double>(Key, Value), InterpMode);
@@ -238,11 +282,29 @@
 	}
 }
 
-void UProceduralDaySequenceBuilder::ClearScalarKeys(FName PropertyName)
+void UProceduralDaySequenceBuilder::AddVectorOverride(FName PropertyName, FVector Value)
 {
+	TPair<float, FVector> A = {0.f, Value};
+	TPair<float, FVector> B = {1.f, Value};
+
+	AddVectorKeys(PropertyName, {A, B}, RCIM_Linear);
+}
+
+void UProceduralDaySequenceBuilder::AddVectorKey(FName PropertyName, float Key, FVector Value, ERichCurveInterpMode InterpMode)
+{
+	AddVectorKey(PropertyName, TPair<float, FVector>(Key, Value), InterpMode);
+}
+
+void UProceduralDaySequenceBuilder::AddVectorKey(FName PropertyName, const TPair<float, FVector>& KeyValue, ERichCurveInterpMode InterpMode)
+{
+	AddVectorKeys(PropertyName, TArray {KeyValue}, InterpMode);
+}
+
+void UProceduralDaySequenceBuilder::AddVectorKeys(FName PropertyName, const TArray<TPair<float, FVector>>& KeysAndValues, ERichCurveInterpMode InterpMode)
+{
 	if (!IsInitialized())
 	{
-		FFrame::KismetExecutionMessage(TEXT("ClearScalarKeys called on an uninitialized Procedural Day Sequence Builder!"), ELogVerbosity::Error);
+		FFrame::KismetExecutionMessage(TEXT("AddVectorKey(s) called on an uninitialized Procedural Day Sequence Builder!"), ELogVerbosity::Error);
 		return;
 	}
 	
@@ -254,73 +316,77 @@
 	{
 		// Do nothing
 	}
-	else if (Property->IsA<FFloatProperty>())
+	else if (Property->IsA<FStructProperty>() && CastField<FStructProperty>(Property)->Struct == TBaseStructure<FVector>::Get())
 	{
-		UMovieSceneFloatSection* FloatSection = CreateOrAddPropertyOverrideSection<UMovieSceneFloatTrack, UMovieSceneFloatSection>(PropertyName);
+		UMovieSceneDoubleVectorSection* VectorSection = CreateOrAddPropertyOverrideSection<UMovieSceneDoubleVectorTrack, UMovieSceneDoubleVectorSection>(PropertyName);
+		VectorSection->SetChannelsUsed(3);
+
+		FMovieSceneDoubleChannel* X = VectorSection->GetChannelProxy().GetChannel<FMovieSceneDoubleChannel>(0);
+		FMovieSceneDoubleChannel* Y = VectorSection->GetChannelProxy().GetChannel<FMovieSceneDoubleChannel>(1);
+		FMovieSceneDoubleChannel* Z = VectorSection->GetChannelProxy().GetChannel<FMovieSceneDoubleChannel>(2);
 		
-		FloatSection->GetChannel().Reset();
+		for (const TPair<float, FVector>& KeyValue : KeysAndValues)
+		{
+			const FFrameNumber FrameNumber = UE::DaySequence::GetKeyFrameNumber(KeyValue.Key, MovieScene->GetPlaybackRange());
+			
+			UE::DaySequence::AddDoubleKey(FrameNumber, KeyValue.Value.X, X, InterpMode);
+			UE::DaySequence::AddDoubleKey(FrameNumber, KeyValue.Value.X, Y, InterpMode);
+			UE::DaySequence::AddDoubleKey(FrameNumber, KeyValue.Value.X, Z, InterpMode);
+		}
 
-		FloatSection->MarkAsChanged();
-	}
-	else if (Property->IsA<FDoubleProperty>())
-	{
-		UMovieSceneDoubleSection* DoubleSection = CreateOrAddPropertyOverrideSection<UMovieSceneDoubleTrack, UMovieSceneDoubleSection>(PropertyName);
-		
-		DoubleSection->GetChannel().Reset();
-		
-		DoubleSection->MarkAsChanged();
-	}
-	else
-	{
-		FFrame::KismetExecutionMessage(*FString::Printf(TEXT("Property %s is not a scalar."), *Property->GetClass()->GetName()), ELogVerbosity::Error);
+		VectorSection->MarkAsChanged();
 	}
 }
 
-void UProceduralDaySequenceBuilder::AddBoolKey(FName PropertyName, float Key, bool Value)
+void UProceduralDaySequenceBuilder::AddColorOverride(FName PropertyName, FLinearColor Value)
 {
-	AddBoolKey(PropertyName, TPair<float, bool>(Key, Value));
+	TPair<float, FLinearColor> A = {0.f, Value};
+	TPair<float, FLinearColor> B = {1.f, Value};
+
+	AddColorKeys(PropertyName, {A, B}, RCIM_Linear);
 }
 
-void UProceduralDaySequenceBuilder::AddBoolKey(FName PropertyName, const TPair<float, bool>& KeyValue)
+void UProceduralDaySequenceBuilder::AddColorKeys(FName PropertyName, const TArray<TPair<float, FLinearColor>>& KeysAndValues, ERichCurveInterpMode InterpMode)
 {
-	AddBoolKeys(PropertyName, TArray {KeyValue});
-}
-
-void UProceduralDaySequenceBuilder::AddBoolKeys(FName PropertyName, const TArray<TPair<float, bool>>& KeysAndValues)
-{
 	if (!IsInitialized())
 	{
-		FFrame::KismetExecutionMessage(TEXT("AddBoolKey(s) called on an uninitialized Procedural Day Sequence Builder!"), ELogVerbosity::Error);
+		FFrame::KismetExecutionMessage(TEXT("AddColorKey(s) called on an uninitialized Procedural Day Sequence Builder!"), ELogVerbosity::Error);
 		return;
 	}
 	
 	const UMovieScene* MovieScene = ProceduralDaySequence->GetMovieScene();
-	
+
 	FTrackInstancePropertyBindings Bindings(PropertyName, PropertyName.ToString());
 	FProperty* Property = Bindings.GetProperty(*ActiveBoundObject);
 	if (!UE::DaySequence::IsPropertyValid(ActiveBoundObject, Property))
 	{
 		// Do nothing
 	}
-	else if (Property->IsA<FBoolProperty>())
+	else if (Property->IsA<FStructProperty>() && 
+		(CastField<FStructProperty>(Property)->Struct == TBaseStructure<FLinearColor>::Get() || CastField<FStructProperty>(Property)->Struct == TBaseStructure<FColor>::Get()) )
 	{
-		UMovieSceneBoolSection* Section = CreateOrAddPropertyOverrideSection<UMovieSceneBoolTrack, UMovieSceneBoolSection>(PropertyName);
-		
-		for (const TPair<float, bool>& KeyValue : KeysAndValues)
+		UMovieSceneColorSection* ColorSection = CreateOrAddPropertyOverrideSection<UMovieSceneColorTrack, UMovieSceneColorSection>(PropertyName);
+
+		for (const TPair<float, FLinearColor>& KeyValue : KeysAndValues)
 		{
 			const FFrameNumber FrameNumber = UE::DaySequence::GetKeyFrameNumber(KeyValue.Key, MovieScene->GetPlaybackRange());
-			
-			Section->GetChannel().AddKeys(TArray {FrameNumber}, TArray {KeyValue.Value});
+
+			UE::DaySequence::AddFloatKey(FrameNumber, KeyValue.Value.R, &ColorSection->GetRedChannel(), InterpMode);
+			UE::DaySequence::AddFloatKey(FrameNumber, KeyValue.Value.G, &ColorSection->GetGreenChannel(), InterpMode);
+			UE::DaySequence::AddFloatKey(FrameNumber, KeyValue.Value.B, &ColorSection->GetBlueChannel(), InterpMode);
+			UE::DaySequence::AddFloatKey(FrameNumber, KeyValue.Value.A, &ColorSection->GetAlphaChannel(), InterpMode);
 		}
-
-		Section->MarkAsChanged();
-	}
-	else
-	{
-		FFrame::KismetExecutionMessage(*FString::Printf(TEXT("Unable to animate a %s property as a bool."), *Property->GetClass()->GetName()), ELogVerbosity::Error);
+		
+		ColorSection->MarkAsChanged();
 	}
 }
 
+void UProceduralDaySequenceBuilder::AddTransformOverride(const FTransform& Value)
+{
+	AddTransformKey(0.f, Value, RCIM_Linear);
+	AddTransformKey(1.f, Value, RCIM_Linear);
+}
+
 void UProceduralDaySequenceBuilder::AddTransformKey(float Key, const FTransform& Value, ERichCurveInterpMode InterpMode)
 {
 	using namespace UE::DaySequence;
@@ -393,18 +459,155 @@
 	AddDoubleKey(FrameNumber, Value.Z, TransformSection->GetChannelProxy().GetChannel<FMovieSceneDoubleChannel>(8), InterpMode);
 }
 
-void UProceduralDaySequenceBuilder::AddStaticTime(float StaticTime)
+void UProceduralDaySequenceBuilder::AddMaterialOverride(int32 MaterialIndex, UMaterialInterface* Value)
+{
+	if (!IsInitialized())
+    {
+    	FFrame::KismetExecutionMessage(TEXT("AddMaterialOverride called on an uninitialized Procedural Day Sequence Builder!"), ELogVerbosity::Error);
+    	return;
+    }
+	
+	UMovieScenePrimitiveMaterialTrack* MaterialTrack = CreateOrAddOverrideTrack<UMovieScenePrimitiveMaterialTrack>(FName());
+	MaterialTrack->SetMaterialInfo(FComponentMaterialInfo{ FName(), MaterialIndex, EComponentMaterialType::IndexedMaterial });
+
+	UMovieScenePrimitiveMaterialSection* Section = Cast<UMovieScenePrimitiveMaterialSection>(MaterialTrack->GetAllSections()[0]);
+	Section->MaterialChannel.SetDefault(Value);
+}
+
+void UProceduralDaySequenceBuilder::AddScalarMaterialParameterOverride(FName ParameterName, int32 MaterialIndex, float Value)
+{
+	TPair<float, float> A = {0.f, Value};
+	TPair<float, float> B = {1.f, Value};
+
+	AddScalarMaterialParameterKeys(ParameterName, MaterialIndex, {A, B});
+}
+
+void UProceduralDaySequenceBuilder::AddScalarMaterialParameterKeys(FName ParameterName, int32 MaterialIndex, const TArray<TPair<float, float>>& KeysAndValues)
 {
 	if (!IsInitialized())
 	{
-		FFrame::KismetExecutionMessage(TEXT("AddStaticTime called on an uninitialized Procedural Day Sequence Builder!"), ELogVerbosity::Error);
+		FFrame::KismetExecutionMessage(TEXT("AddScalarMaterialParameterKeys called on an uninitialized Procedural Day Sequence Builder!"), ELogVerbosity::Error);
 		return;
 	}
 	
-	UMovieSceneFloatSection* Section = CreateOrAddPropertyOverrideSection<UMovieSceneFloatTrack, UMovieSceneFloatSection>("StaticTimeOfDay");
-	Section->GetChannel().SetDefault(StaticTime);
+	const UMovieScene* MovieScene = ProceduralDaySequence->GetMovieScene();
+
+	// Material parameter tracks use the material index as the unique name
+	const FName IndexAsName(*FString::FromInt(MaterialIndex));
+	UMovieSceneComponentMaterialTrack* MaterialTrack = CreateOrAddOverrideTrack<UMovieSceneComponentMaterialTrack>(IndexAsName);
+	MaterialTrack->SetMaterialInfo(FComponentMaterialInfo{FName(), MaterialIndex, EComponentMaterialType::IndexedMaterial });
+
+	for (const TPair<float, float>& KeyValue : KeysAndValues)
+	{
+		const FFrameNumber FrameNumber = UE::DaySequence::GetKeyFrameNumber(KeyValue.Key, MovieScene->GetPlaybackRange());
+		
+		MaterialTrack->AddScalarParameterKey(ParameterName, FrameNumber, KeyValue.Value);
+	}
 }
 
+void UProceduralDaySequenceBuilder::AddColorMaterialParameterOverride(FName ParameterName, int32 MaterialIndex, FLinearColor Value)
+{
+	TPair<float, FLinearColor> A = {0.f, Value};
+	TPair<float, FLinearColor> B = {1.f, Value};
+
+	AddColorMaterialParameterKeys(ParameterName, MaterialIndex, {A, B});
+}
+
+void UProceduralDaySequenceBuilder::AddColorMaterialParameterKeys(FName ParameterName, int32 MaterialIndex, const TArray<TPair<float, FLinearColor>>& KeysAndValues)
+{
+	if (!IsInitialized())
+	{
+		FFrame::KismetExecutionMessage(TEXT("AddColorMaterialParameterKeys called on an uninitialized Procedural Day Sequence Builder!"), ELogVerbosity::Error);
+		return;
+	}
+	
+	const UMovieScene* MovieScene = ProceduralDaySequence->GetMovieScene();
+
+	// Material parameter tracks use the material index as the unique name
+	const FName IndexAsName(*FString::FromInt(MaterialIndex));
+	UMovieSceneComponentMaterialTrack* MaterialTrack = CreateOrAddOverrideTrack<UMovieSceneComponentMaterialTrack>(IndexAsName);
+	MaterialTrack->SetMaterialInfo(FComponentMaterialInfo{FName(), MaterialIndex, EComponentMaterialType::IndexedMaterial });
+
+	for (const TPair<float, FLinearColor>& KeyValue : KeysAndValues)
+	{
+		const FFrameNumber FrameNumber = UE::DaySequence::GetKeyFrameNumber(KeyValue.Key, MovieScene->GetPlaybackRange());
+		
+		MaterialTrack->AddColorParameterKey(ParameterName, FrameNumber, KeyValue.Value);
+	}
+}
+
+void UProceduralDaySequenceBuilder::AddVisibilityOverride(bool bValue)
+{
+	TPair<float, bool> A = {0.f, bValue};
+	TPair<float, bool> B = {1.f, bValue};
+
+	AddVisibilityKeys({A, B});
+}
+
+void UProceduralDaySequenceBuilder::AddVisibilityKeys(const TArray<TPair<float, bool>>& KeysAndValues)
+{
+	if (!IsInitialized())
+	{
+		FFrame::KismetExecutionMessage(TEXT("AddVisibilityKey(s) called on an uninitialized Procedural Day Sequence Builder!"), ELogVerbosity::Error);
+		return;
+	}
+	
+	const UMovieScene* MovieScene = ProceduralDaySequence->GetMovieScene();
+	
+	static const FName ActorVisibilityTrackName = TEXT("bHidden");
+	static const FName ComponentVisibilityTrackName = TEXT("bHiddenInGame");
+	
+	const bool bIsComponent = ActiveBoundObject->IsA<USceneComponent>();
+	const bool bIsActor     = ActiveBoundObject->IsA<AActor>();
+	
+	if (!bIsComponent && !bIsActor)
+	{
+		FFrame::KismetExecutionMessage(TEXT("AddVisibilityKey(s) called but ActiveBoundObject is neither an Actor nor a Scene Component!"), ELogVerbosity::Error);
+		return;
+	}
+
+	// We can check just bIsComponent because we early either if _both_ are false, so it must be one or the other.
+	const FName& TrackName = bIsComponent ? ComponentVisibilityTrackName : ActorVisibilityTrackName;
+	UMovieSceneVisibilitySection* VisibilitySection = CreateOrAddPropertyOverrideSection<UMovieSceneVisibilityTrack, UMovieSceneVisibilitySection>(TrackName);
+
+	for (const TPair<float, bool>& KeyValue : KeysAndValues)
+	{
+		const FFrameNumber FrameNumber = UE::DaySequence::GetKeyFrameNumber(KeyValue.Key, MovieScene->GetPlaybackRange());
+		
+		VisibilitySection->GetChannel().AddKeys(TArray {FrameNumber}, TArray {KeyValue.Value});
+	}
+	
+	VisibilitySection->MarkAsChanged();
+}
+
+void UProceduralDaySequenceBuilder::ClearKeys()
+{
+	if (!ProceduralDaySequence)
+	{
+		return;
+	}
+	
+	if (UMovieScene* MovieScene = ProceduralDaySequence->GetMovieScene())
+	{
+		for (const FMovieSceneBinding& Binding : MovieScene->GetBindings())
+		{
+			// Inconvenient we have to do this but at least FindBinding is doing a binary search and we do this once per binding.
+			if (FMovieSceneBinding* MutableBinding = MovieScene->FindBinding(Binding.GetObjectGuid()))
+			{
+				// We have to copy the array here because we are mutating the internal array
+				for (TArray<UMovieSceneTrack*> Tracks = MutableBinding->GetTracks(); UMovieSceneTrack* Track : Tracks)
+				{
+					MutableBinding->RemoveTrack(*Track, MovieScene);
+				}
+			}
+		}
+
+		MovieScene->MarkAsChanged();
+	}
+}
+
+
+
 FGuid UProceduralDaySequenceBuilder::GetOrCreateProceduralBinding(UObject* Object) const
 {
 	if (!Object)
diff --git a/Engine/Plugins/Experimental/DaySequence/Source/DaySequence/Public/DaySequenceModifierComponent.h b/Engine/Plugins/Experimental/DaySequence/Source/DaySequence/Public/DaySequenceModifierComponent.h
--- a/Engine/Plugins/Experimental/DaySequence/Source/DaySequence/Public/DaySequenceModifierComponent.h
+++ b/Engine/Plugins/Experimental/DaySequence/Source/DaySequence/Public/DaySequenceModifierComponent.h
@@ -164,115 +164,27 @@ public:
 	 */
 	UFUNCTION(BlueprintCallable, Category="Day Sequence")
 	DAYSEQUENCE_API void DisableComponent();
-	
-	/**
-	 * Remove all the Sequencer tracks within our procedural Day Sequence
-	 */
-	UFUNCTION(BlueprintCallable, Category="Day Sequence")
-	DAYSEQUENCE_API void ResetOverrides();
-
-	/**
-	 * Add a new override for the static time of day on the day-sequence actor
-	 *
-	 * @param Hours         The static time of day in hours to use
-	 */
-	UFUNCTION(BlueprintCallable, Category="Day Sequence")
-	DAYSEQUENCE_API void AddStaticTimeOfDayOverride(ADaySequenceActor* Actor, float Hours);
-
-	/**
-	 * Add a new boolean override for the specified property
-	 *
-	 * @param Object        The object that contains the property we want to override. Must either be the ADaySequenceActor or one of its components
-	 * @param PropertyName  The name of the property to override
-	 * @param Value         The value to override
-	 */
-	UFUNCTION(BlueprintCallable, Category="Day Sequence")
-	DAYSEQUENCE_API void AddBoolOverride(UObject* Object, FName PropertyName, bool bValue);
-
-	/**
-	 * Add a new scalar override for the specified property
-	 *
-	 * @param Object        The object that contains the property we want to override. Must either be the ADaySequenceActor or one of its components
-	 * @param PropertyName  The name of the property to override
-	 * @param Value         The value to override
-	 */
-	UFUNCTION(BlueprintCallable, Category="Day Sequence")
-	DAYSEQUENCE_API void AddScalarOverride(UObject* Object, FName PropertyName, double Value);
-
-	/**
-	 * Add a new vector override for the specified property
-	 *
-	 * @param Object        The object that contains the property we want to override. Must either be the ADaySequenceActor or one of its components
-	 * @param PropertyName  The name of the property to override
-	 * @param Value         The value to override
-	 */
-	UFUNCTION(BlueprintCallable, Category="Day Sequence")
-	DAYSEQUENCE_API void AddVectorOverride(UObject* Object, FName PropertyName, FVector Value);
-
-	/**
-	 * Add a new color override for the specified property
-	 *
-	 * @param Object        The object that contains the property we want to override. Must either be the ADaySequenceActor or one of its components
-	 * @param PropertyName  The name of the property to override.
-	 * @param Value         The value to override
-	 */
-	UFUNCTION(BlueprintCallable, Category="Day Sequence")
-	DAYSEQUENCE_API void AddColorOverride(UObject* Object, FName PropertyName, FLinearColor Value);
-
-	/**
-	 * Add a new transform override for the specified property
-	 *
-	 * @param Object        The object that contains the property we want to override. Must either be the ADaySequenceActor or one of its components
-	 * @param Value         The value to override
-	 */
-	UFUNCTION(BlueprintCallable, Category="Day Sequence")
-	DAYSEQUENCE_API void AddTransformOverride(UObject* Object, FTransform Value);
 
-	/**
-	 * Add a new material override for the specified material element index
-	 *
-	 * @param Object          The object that contains the property we want to override. Must either be the ADaySequenceActor or one of its components
-	 * @param MaterialIndex   The index of the material element to override
-	 * @param Value           The material to assign
-	 */
-	UFUNCTION(BlueprintCallable, Category="Day Sequence")
-	DAYSEQUENCE_API void AddMaterialOverride(UObject* Object, int32 MaterialIndex, UMaterialInterface* Value);
+	/** Sets the user day sequence. This must be a user created asset. */
+	UFUNCTION(BlueprintSetter, Category="Day Sequence")
+	DAYSEQUENCE_API void SetUserDaySequence(UDaySequence* InDaySequence);
 
-	/**
-	 * Add a new scalar material parameter override for the specified material element index and parameter name
-	 *
-	 * @param Object          The object that contains the material we want to override. Must either be the ADaySequenceActor or one of its components
-	 * @param MaterialIndex   The index of the material element to override
-	 * @param ParameterName   The name of the parameter to override
-	 * @param Value           The scalar value to assign to the material parameter
-	 */
-	UFUNCTION(BlueprintCallable, Category="Day Sequence")
-	DAYSEQUENCE_API void AddScalarMaterialParameterOverride(UObject* Object, int32 MaterialIndex, FName ParameterName, float Value);
+	UFUNCTION(BlueprintGetter, BlueprintPure, Category="Day Sequence")
+	DAYSEQUENCE_API UDaySequence* GetUserDaySequence();
 
-	/**
-	 * Add a new color material parameter override for the specified material element index and parameter name
-	 *
-	 * @param Object          The object that contains the material we want to override. Must either be the ADaySequenceActor or one of its components
-	 * @param MaterialIndex   The index of the material element to override
-	 * @param ParameterName   The name of the parameter to override
-	 * @param Value           The color value to assign to the material parameter
-	 */
-	UFUNCTION(BlueprintCallable, Category="Day Sequence")
-	DAYSEQUENCE_API void AddColorMaterialParameterOverride(UObject* Object, int32 MaterialIndex, FName ParameterName, FLinearColor Value);
-
-	/**
-	 * Add a new visibility override for the specified property
-	 *
-	 * @param Object        The object that contains the property we want to override. Must either be the ADaySequenceActor or one of its components
-	 * @param Value         The value to override
-	 */
+	/** Sets a user specified transient sequence. */
+	UFUNCTION(BlueprintSetter, Category="Day Sequence")
+	DAYSEQUENCE_API void SetTransientSequence(UDaySequence* InDaySequence);
+	
+	UFUNCTION(BlueprintGetter, BlueprintPure, Category="Day Sequence")
+	DAYSEQUENCE_API UDaySequence* GetTransientSequence();
+	
 	UFUNCTION(BlueprintCallable, Category="Day Sequence")
-	DAYSEQUENCE_API void AddVisibilityOverride(UObject* Object, bool bValue);
+	DAYSEQUENCE_API void SetDayNightCycle(EDayNightCycleMode NewMode);
 
-	/** Sets the user day sequence */
 	UFUNCTION(BlueprintCallable, Category="Day Sequence")
-	DAYSEQUENCE_API void SetUserDaySequence(UDaySequence* InDaySequence);
-
+	DAYSEQUENCE_API void SetDayNightCycleTime(float Time);
+	
 	UFUNCTION(BlueprintCallable, Category="Day Sequence")
 	DAYSEQUENCE_API void SetMode(EDaySequenceModifierMode NewMode);
 
@@ -293,35 +205,25 @@ public:
 
 	void EmptyVolumeShapeComponents();
 	void AddVolumeShapeComponent(const FComponentReference& InShapeReference);
-
-	void InvalidateMuteStates() const;
-
-#if ENABLE_DRAW_DEBUG
-	bool ShouldShowDebugInfo() const;
-#endif
 	
 	/*~ Begin UActorComponent interface */
-	void BeginPlay() override;
-	void EndPlay(EEndPlayReason::Type Reason) override;
-	void OnRegister() override;
-	void OnUnregister() override;
+	virtual void BeginPlay() override;
+	virtual void EndPlay(EEndPlayReason::Type Reason) override;
+	virtual void OnRegister() override;
+	virtual void OnUnregister() override;
 	/*~ End UActorComponent interface */
-
-	/** Bound to delegate on the DaySequenceActor that allows all modifiers to do work at appropriate times at the specific actors tick interval. */
-	void DaySequenceUpdate();
 	
 #if WITH_EDITOR
 	/*~ Begin FTickableGameObject interface */
-	void UpdateEditorPreview(float DeltaTime) override;
-	TStatId GetStatId() const override;
-	ETickableTickType GetTickableTickType() const override;
-	bool IsTickableWhenPaused() const override { return true; }
-	bool IsTickableInEditor() const override { return true; }
-	bool IsTickable() const override { return true; }
+	virtual void UpdateEditorPreview(float DeltaTime) override;
+	virtual TStatId GetStatId() const override;
+	virtual ETickableTickType GetTickableTickType() const override;
+	virtual bool IsTickableWhenPaused() const override { return true; }
+	virtual bool IsTickableInEditor() const override { return true; }
+	virtual bool IsTickable() const override { return true; }
 	/*~ End FTickableGameObject interface */
 
-	void PostEditChangeProperty(FPropertyChangedEvent& PropertyChangedEvent) override;
-
+	virtual void PostEditChangeProperty(FPropertyChangedEvent& PropertyChangedEvent) override;
 #endif
 
 protected:
@@ -334,13 +236,13 @@ protected:
 	
 	bool CanBeEnabled() const;
 
-	UPROPERTY(Instanced, Transient)
-	TObjectPtr<UDaySequenceModifierEasingFunction> EasingFunction;
-
 	TArray<UShapeComponent*> GetVolumeShapeComponents() const;
 	
 	void SetInitialTimeOfDay();
 
+	/* Called to properly update the mute states of all managed subsections. */
+	void InvalidateMuteStates() const;
+	
 protected:
 
 	/** Non-serialized target actor we are currently bound to */
@@ -358,16 +260,20 @@ protected:
 	UPROPERTY(Transient)
 	TWeakObjectPtr<APlayerController> WeakBlendTarget;
 
-	/** An optional user-provided Day Sequence - used instead of our procedurally generated one if set */
-	UPROPERTY(BlueprintReadWrite, EditAnywhere, Category="Day Sequence", meta=(EditCondition="!bUseCollection", EditConditionHides, DisplayAfter="bUseCollection"))
+	/** The user provided Day Sequence. */
+	UPROPERTY(EditAnywhere, Category="Day Sequence", BlueprintGetter = GetUserDaySequence, BlueprintSetter = SetUserDaySequence, meta=(EditCondition="!bUseCollection", EditConditionHides, DisplayAfter="bUseCollection"))
 	TObjectPtr<UDaySequence> UserDaySequence;
 
-	/** The procedurally generated sequence containing our override tracks. Owned by this component. */
-	UPROPERTY(Transient)
-	TObjectPtr<UDaySequence> ProceduralDaySequence;
-
+	/** The user provided Transient Day Sequence. */
+    UPROPERTY(Transient, Category="Day Sequence", BlueprintGetter = GetTransientSequence, BlueprintSetter = SetTransientSequence, meta=(EditCondition="!bUseCollection", EditConditionHides, DisplayAfter="bUseCollection"))
+    TObjectPtr<UDaySequence> TransientSequence;
+	
+	/** The user provided collection. This is an alternative to UserDaySequence. */
 	UPROPERTY(EditAnywhere, Category="Day Sequence", meta=(EditCondition="bUseCollection", EditConditionHides, DisplayAfter="bUseCollection"))
 	TObjectPtr<UDaySequenceCollectionAsset> DaySequenceCollection;
+
+	UPROPERTY(Instanced, Transient)
+	TObjectPtr<UDaySequenceModifierEasingFunction> EasingFunction;
 	
 	/** User-defined bias. */
 	UPROPERTY(BlueprintReadWrite, EditAnywhere, Category="Day Sequence", meta=(EditCondition="!bIgnoreBias"))
@@ -434,6 +340,12 @@ protected:
 
 private:
 	
+	/**
+	 * Bound to delegate on the DaySequenceActor that allows all modifiers to do work at appropriate times at the specific actors tick interval.
+	 * Effectively a 'tick' function.
+	 */
+	void DaySequenceUpdate();
+	
 	/**
 	 * Get the blend position (handles preview and game world).
 	 * @return Returns true if we have a valid blend position and InPosition was set, false if InPosition was not set.
@@ -460,14 +372,12 @@ private:
 	void ReinitializeSubSequence(ADaySequenceActor::FSubSectionPreserveMap* SectionsToPreserve);
 	UMovieSceneSubSection* InitializeDaySequence(const FDaySequenceCollectionEntry& SequenceAsset);
 	void RemoveSubSequenceTrack();
-	FGuid GetOrCreateProceduralBinding(UObject* Object);
 
 	void UpdateCachedExternalShapes() const;
 	mutable TArray<TWeakObjectPtr<UShapeComponent>> CachedExternalShapes;
 	mutable bool bCachedExternalShapesInvalid = true;
 	
 	/*~ Transient state for active gameplay */
-	TWeakObjectPtr<UMovieSceneSubSection> WeakSubSection;
 	TArray<TWeakObjectPtr<UMovieSceneSubSection>> SubSections;
 
 	UE::DaySequence::FOnInvalidateMuteStates OnInvalidateMuteStates;
@@ -476,6 +386,7 @@ private:
 	const FName ShowDebug_ModifierCategory = "DaySequence_Modifiers";
 	
 	void OnDebugLevelChanged(int32 InDebugLevel);
+	bool ShouldShowDebugInfo() const;
 	
 	/** Determines whether or not the modifier will show debug info */
 	int32 DebugLevel;
diff --git a/Engine/Plugins/Experimental/DaySequence/Source/DaySequence/Public/ProceduralDaySequenceBuilder.h b/Engine/Plugins/Experimental/DaySequence/Source/DaySequence/Public/ProceduralDaySequenceBuilder.h
--- a/Engine/Plugins/Experimental/DaySequence/Source/DaySequence/Public/ProceduralDaySequenceBuilder.h
+++ b/Engine/Plugins/Experimental/DaySequence/Source/DaySequence/Public/ProceduralDaySequenceBuilder.h
@@ -30,7 +30,7 @@ public:
 	 * 
 	 * @param InActor The target DaySequenceActor that will be animated by the generated sequence.
 	 * @param InitialSequence Optional sequence that this builder can operate on instead of allocating a new sequence.
-	 * @param bClearInitialSequence If true, calls ClearAllKeys().
+	 * @param bClearInitialSequence If true, calls ClearKeys().
 	 * @return The sequence which will be modified when calling SetActiveBoundObject and the Add*Key(s) functions.
 	 */
 	UFUNCTION(BlueprintCallable, Category = "Day Sequence")
@@ -44,38 +44,72 @@ public:
 	UFUNCTION(BlueprintCallable, Category = "Day Sequence")
 	void SetActiveBoundObject(UObject* InObject);
 
-	// todo [nickolas.drake]: Fill out this section, making sure common key types can be added in bulk and by BPs
 	
+	/* Key Creation: */
+
+	UFUNCTION(BlueprintCallable, Category = "Day Sequence")
+	void AddBoolOverride(FName PropertyName, bool Value);
 	UFUNCTION(BlueprintCallable, Category = "Day Sequence")
-	void ClearAllKeys();
+	void AddBoolKey(FName PropertyName, float Key, bool Value);
+	void AddBoolKey(FName PropertyName, const TPair<float, bool>& KeyValue);
+	void AddBoolKeys(FName PropertyName, const TArray<TPair<float, bool>>& KeysAndValues);
+	
 	
+	UFUNCTION(BlueprintCallable, Category = "Day Sequence")
+	void AddScalarOverride(FName PropertyName, double Value);
 	UFUNCTION(BlueprintCallable, Category = "Day Sequence")
 	void AddScalarKey(FName PropertyName, float Key, double Value, ERichCurveInterpMode InterpMode = RCIM_Cubic);
 	void AddScalarKey(FName PropertyName, const TPair<float, double>& KeyValue, ERichCurveInterpMode InterpMode = RCIM_Cubic);
 	void AddScalarKeys(FName PropertyName, const TArray<TPair<float, double>>& KeysAndValues, ERichCurveInterpMode InterpMode = RCIM_Cubic);
+
 	
 	UFUNCTION(BlueprintCallable, Category = "Day Sequence")
-	void ClearScalarKeys(FName PropertyName);
-
+	void AddVectorOverride(FName PropertyName, FVector Value);
 	UFUNCTION(BlueprintCallable, Category = "Day Sequence")
-	void AddBoolKey(FName PropertyName, float Key, bool Value);
-	void AddBoolKey(FName PropertyName, const TPair<float, bool>& KeyValue);
-	void AddBoolKeys(FName PropertyName, const TArray<TPair<float, bool>>& KeysAndValues);
+	void AddVectorKey(FName PropertyName, float Key, FVector Value, ERichCurveInterpMode InterpMode = RCIM_Cubic);
+	void AddVectorKey(FName PropertyName, const TPair<float, FVector>& KeyValue, ERichCurveInterpMode InterpMode = RCIM_Cubic);
+	void AddVectorKeys(FName PropertyName, const TArray<TPair<float, FVector>>& KeysAndValues, ERichCurveInterpMode InterpMode = RCIM_Cubic);
 
+	
 	UFUNCTION(BlueprintCallable, Category = "Day Sequence")
-	void AddTransformKey(float Key, const FTransform& Value, ERichCurveInterpMode InterpMode = RCIM_Cubic);
+	void AddColorOverride(FName PropertyName, FLinearColor Value);
+	void AddColorKeys(FName PropertyName, const TArray<TPair<float, FLinearColor>>& KeysAndValues, ERichCurveInterpMode InterpMode = RCIM_Cubic);
+	
 	
+	UFUNCTION(BlueprintCallable, Category = "Day Sequence")
+	void AddTransformOverride(const FTransform& Value);
+	void AddTransformKey(float Key, const FTransform& Value, ERichCurveInterpMode InterpMode = RCIM_Cubic);
 	UFUNCTION(BlueprintCallable, Category = "Day Sequence")
 	void AddTranslationKey(float Key, const FVector& Value, ERichCurveInterpMode InterpMode = RCIM_Cubic);
-	
 	UFUNCTION(BlueprintCallable, Category = "Day Sequence")
 	void AddRotationKey(float Key, const FRotator& Value, ERichCurveInterpMode InterpMode = RCIM_Cubic);
-	
 	UFUNCTION(BlueprintCallable, Category = "Day Sequence")
 	void AddScaleKey(float Key, const FVector& Value, ERichCurveInterpMode InterpMode = RCIM_Cubic);
+	
+
+	UFUNCTION(BlueprintCallable, Category = "Day Sequence")
+	void AddMaterialOverride(int32 MaterialIndex, UMaterialInterface* Value);
 
+	
+	UFUNCTION(BlueprintCallable, Category="Day Sequence")
+	void AddScalarMaterialParameterOverride(FName ParameterName, int32 MaterialIndex, float Value);
+	void AddScalarMaterialParameterKeys(FName ParameterName, int32 MaterialIndex, const TArray<TPair<float, float>>& KeysAndValues);
+
+	
+	UFUNCTION(BlueprintCallable, Category="Day Sequence")
+	void AddColorMaterialParameterOverride(FName ParameterName, int32 MaterialIndex, FLinearColor Value);
+	void AddColorMaterialParameterKeys(FName ParameterName, int32 MaterialIndex, const TArray<TPair<float, FLinearColor>>& KeysAndValues);
+
+
+	UFUNCTION(BlueprintCallable, Category="Day Sequence")
+	void AddVisibilityOverride(bool bValue);
+	void AddVisibilityKeys(const TArray<TPair<float, bool>>& KeysAndValues);
+	
+	
+	/* Key Deletion: */
+	
 	UFUNCTION(BlueprintCallable, Category = "Day Sequence")
-	void AddStaticTime(float StaticTime);
+	void ClearKeys();
 	
 private:
 	FGuid GetOrCreateProceduralBinding(UObject* Object) const;
