DaySequence: Let modifier collections use procedural sequences and make the condition set picker work with cooked BP generated classes

diff --git a/Engine/Plugins/Experimental/DaySequence/Source/DaySequence/Private/DaySequenceModifierComponent.cpp b/Engine/Plugins/Experimental/DaySequence/Source/DaySequence/Private/DaySequenceModifierComponent.cpp
--- a/Engine/Plugins/Experimental/DaySequence/Source/DaySequence/Private/DaySequenceModifierComponent.cpp
+++ b/Engine/Plugins/Experimental/DaySequence/Source/DaySequence/Private/DaySequenceModifierComponent.cpp
@@ -637,6 +637,16 @@ void UDaySequenceModifierComponent::ReinitializeSubSequence(ADaySequenceActor::F
 				{
 					InitializeDaySequence(Entry);
 				}
+
+				for (TInstancedStruct<FProceduralDaySequence>& ProceduralEntry : DaySequenceCollection->ProceduralDaySequences)
+				{
+					if (!ProceduralEntry.IsValid())
+					{
+						continue;
+					}
+			
+					InitializeDaySequence(ProceduralEntry.GetMutable<FProceduralDaySequence>().GetSequence(TargetActor));
+				}
 			}
 		}
 		else
diff --git a/Engine/Plugins/Experimental/DaySequence/Source/DaySequenceEditor/Private/SDaySequenceConditionSetPicker.cpp b/Engine/Plugins/Experimental/DaySequence/Source/DaySequenceEditor/Private/SDaySequenceConditionSetPicker.cpp
--- a/Engine/Plugins/Experimental/DaySequence/Source/DaySequenceEditor/Private/SDaySequenceConditionSetPicker.cpp
+++ b/Engine/Plugins/Experimental/DaySequence/Source/DaySequenceEditor/Private/SDaySequenceConditionSetPicker.cpp
@@ -21,6 +21,7 @@
 #include "Blueprint/BlueprintSupport.h"	// For FBlueprintTags::NativeParentClassPath
 #include "AssetRegistry/AssetData.h"
 #include "Engine/Blueprint.h"
+#include "Engine/BlueprintGeneratedClass.h"
 #include "Widgets/Layout/SScrollBar.h"
 #include "Widgets/Views/SListView.h"
 
@@ -183,26 +184,37 @@ void SDaySequenceConditionSetPicker::PopulateVisibleClasses()
 	const FTopLevelAssetPath BaseAssetPath = UDaySequenceConditionTag::StaticClass()->GetClassPathName();
 
 	AllConditionTags.Reset();
-	
-	FARFilter Filter;
-	Filter.ClassPaths.Add(UBlueprint::StaticClass()->GetClassPathName());
-	Filter.bRecursiveClasses = true;
-	Filter.TagsAndValues.Add(FBlueprintTags::NativeParentClassPath);
-	
-	AssetRegistry.EnumerateAssets(Filter, [this, BaseAssetPath](const FAssetData& Data)
+
+	// Scan for BP assets to get generated classes
 	{
-		if (Data.TagsAndValues.FindTag(FBlueprintTags::NativeParentClassPath).AsExportPath().ToTopLevelAssetPath() == BaseAssetPath)
+		FARFilter Filter;
+		Filter.ClassPaths.Add(UBlueprint::StaticClass()->GetClassPathName());
+		Filter.ClassPaths.Add(UBlueprintGeneratedClass::StaticClass()->GetClassPathName());
+		Filter.bRecursiveClasses = true;
+		Filter.TagsAndValues.Add(FBlueprintTags::NativeParentClassPath);
+	
+		AssetRegistry.EnumerateAssets(Filter, [this, BaseAssetPath](const FAssetData& Data)
 		{
-			if (const UBlueprint* AssetAsBlueprint = Cast<UBlueprint>(Data.GetAsset()))
+			if (Data.TagsAndValues.FindTag(FBlueprintTags::NativeParentClassPath).AsExportPath().ToTopLevelAssetPath() == BaseAssetPath)
 			{
-				if (const TSubclassOf<UDaySequenceConditionTag> GeneratedClassAsConditionClass = *AssetAsBlueprint->GeneratedClass)
+				if (const UBlueprint* AssetAsBlueprint = Cast<UBlueprint>(Data.GetAsset()))
 				{
-					AllConditionTags.Add(GeneratedClassAsConditionClass);
+					if (const TSubclassOf<UDaySequenceConditionTag> GeneratedClassAsConditionClass = *AssetAsBlueprint->GeneratedClass)
+					{
+						AllConditionTags.AddUnique(GeneratedClassAsConditionClass);
+					}
+				}
+				else if (UBlueprintGeneratedClass* AssetAsBlueprintGeneratedClass = Cast<UBlueprintGeneratedClass>(Data.GetAsset()))
+				{
+					if (AssetAsBlueprintGeneratedClass->IsChildOf(UDaySequenceConditionTag::StaticClass()))
+					{
+						AllConditionTags.Add(AssetAsBlueprintGeneratedClass);
+					}
 				}
 			}
-		}
-		return true;	// Returning false will halt the enumeration
-	});
+			return true;	// Returning false will halt the enumeration
+		});
+	}
 }
 
 void SDaySequenceConditionSetPicker::PopulateCheckedTags()
