DaySequence: Sample content adjustments

- Put moon at 45 degrees in sky
- Make sun disc smaller
- Animate cloud parameters

#virtualized

diff --git a/Engine/Plugins/Experimental/DaySequence/Source/DaySequence/Private/Actors/SunMoonDaySequenceActor.cpp b/Engine/Plugins/Experimental/DaySequence/Source/DaySequence/Private/Actors/SunMoonDaySequenceActor.cpp
--- a/Engine/Plugins/Experimental/DaySequence/Source/DaySequence/Private/Actors/SunMoonDaySequenceActor.cpp
+++ b/Engine/Plugins/Experimental/DaySequence/Source/DaySequence/Private/Actors/SunMoonDaySequenceActor.cpp
@@ -24,6 +24,7 @@ ASunMoonDaySequenceActor::ASunMoonDaySequenceActor(const FObjectInitializer& Ini
 	MoonComponent->SetIntensity(2.f);
 	MoonComponent->SetUseTemperature(true);
 	MoonComponent->SetTemperature(9000.f);
+	MoonComponent->SetWorldRotation(FRotator(0.f, 45.f, 0.f));
 	
 	// Override the sky sphere material.
 	static ConstructorHelpers::FObjectFinder<UMaterialInterface> SkySphereDefaultMaterial(TEXT("/DaySequence/MI_24hrSky.MI_24hrSky"));
