DaySequenceActor: Replicate playback by default

I think this is reasonable to do by default because a user can just make the actor replicate or not and it will yield expected behavior. In the current form the actor can replicate and it still might not be obvious how to synchronize playback. I think the vast majority of cases will have this set to true.

diff --git a/Engine/Plugins/Experimental/DaySequence/Source/DaySequence/Private/DaySequenceActor.cpp b/Engine/Plugins/Experimental/DaySequence/Source/DaySequence/Private/DaySequenceActor.cpp
--- a/Engine/Plugins/Experimental/DaySequence/Source/DaySequence/Private/DaySequenceActor.cpp
+++ b/Engine/Plugins/Experimental/DaySequence/Source/DaySequence/Private/DaySequenceActor.cpp
@@ -116,7 +116,7 @@ ADaySequenceActor::ADaySequenceActor(const FObjectInitializer& Init)
 	bAlwaysRelevant = true;
 	bReplicates = true;
 	bReplicateUsingRegisteredSubObjectList = true;
-	bReplicatePlayback = false;
+	bReplicatePlayback = true;
 	SetHidden(false);
 	
 #if ENABLE_DRAW_DEBUG
