DaySequence: Fix crash when inspecting modifier subsequences in the root sequence while a level is being validated

I think this crash could probably occur in scenarios other than level validation. The basic issue is as follows:
1) A user is inspecting/focused on a modifier subsequence section
2) The modifier removes its subsequence section from the root sequence
3) An evaluation occurs

The result is a nullptr dereference because much of the core Sequencer code assumes the focused subsequence is always valid. The fix is to introduce a static delegate on ADaySequenceActor with the signature void(UMovieSceneSubSection*) that is broadcast by the modifier when its subsection is removed. The day sequence editor module monitors this delegate and changes the focused sequence to the root sequence if the currently focused subsequence is the same as the one assigned to the section passed as a parameter when invoking the delegate.

diff --git a/Engine/Plugins/Experimental/DaySequence/Source/DaySequence/Private/DaySequenceActor.cpp b/Engine/Plugins/Experimental/DaySequence/Source/DaySequence/Private/DaySequenceActor.cpp
--- a/Engine/Plugins/Experimental/DaySequence/Source/DaySequence/Private/DaySequenceActor.cpp
+++ b/Engine/Plugins/Experimental/DaySequence/Source/DaySequence/Private/DaySequenceActor.cpp
@@ -40,6 +40,10 @@
 
 #include UE_INLINE_GENERATED_CPP_BY_NAME(DaySequenceActor)
 
+#if WITH_EDITOR
+ADaySequenceActor::FOnSubSectionRemovedEvent ADaySequenceActor::OnSubSectionRemovedEvent;
+#endif
+
 namespace UE::DaySequence
 {
 #if ENABLE_DRAW_DEBUG
@@ -572,6 +576,10 @@ void ADaySequenceActor::InitializeRootSequence()
 		{
 			if (!Section.Value)
 			{
+#if WITH_EDITOR
+				OnSubSectionRemovedEvent.Broadcast(Section.Key);
+#endif
+				
 				UMovieSceneTrack* Track = Section.Key->GetTypedOuter<UMovieSceneTrack>();
 				UMovieScene* MovieScene = Track->GetTypedOuter<UMovieScene>();
 	
diff --git a/Engine/Plugins/Experimental/DaySequence/Source/DaySequence/Private/DaySequenceModifierComponent.cpp b/Engine/Plugins/Experimental/DaySequence/Source/DaySequence/Private/DaySequenceModifierComponent.cpp
--- a/Engine/Plugins/Experimental/DaySequence/Source/DaySequence/Private/DaySequenceModifierComponent.cpp
+++ b/Engine/Plugins/Experimental/DaySequence/Source/DaySequence/Private/DaySequenceModifierComponent.cpp
@@ -394,6 +394,9 @@ void UDaySequenceModifierComponent::RemoveSubSequenceTrack()
 	
 	for (const TWeakObjectPtr<UMovieSceneSubSection> SubSection : SubSections)
 	{
+#if WITH_EDITOR
+		ADaySequenceActor::OnSubSectionRemovedEvent.Broadcast(SubSection.Get());
+#endif
 		RemoveSubTrack(SubSection.Get());
 	}
 	SubSections.Empty();
diff --git a/Engine/Plugins/Experimental/DaySequence/Source/DaySequence/Public/DaySequenceActor.h b/Engine/Plugins/Experimental/DaySequence/Source/DaySequence/Public/DaySequenceActor.h
--- a/Engine/Plugins/Experimental/DaySequence/Source/DaySequence/Public/DaySequenceActor.h
+++ b/Engine/Plugins/Experimental/DaySequence/Source/DaySequence/Public/DaySequenceActor.h
@@ -114,6 +114,14 @@ public:
 	 */
 	DECLARE_EVENT_OneParam(ADaySequenceActor, FOnTimeOfDayPreviewChangedEvent, float)
 	FOnTimeOfDayPreviewChangedEvent OnTimeOfDayPreviewChangedEvent;
+
+	/**
+	 * Native event invoked when a subsection is removed from the Root Sequence.
+	 * Primarily called by modifier components but generally callable by anything that adds a subsection to the root sequence.
+	 * @param The subsection being removed.
+	 */
+	DECLARE_EVENT_OneParam(ADaySequenceActor, FOnSubSectionRemovedEvent, const UMovieSceneSubSection*)
+	static FOnSubSectionRemovedEvent OnSubSectionRemovedEvent;
 #endif //WITH_EDITOR
 
 	/** @return the preview time in hours. */
diff --git a/Engine/Plugins/Experimental/DaySequence/Source/DaySequenceEditor/Private/DaySequenceEditorModule.cpp b/Engine/Plugins/Experimental/DaySequence/Source/DaySequenceEditor/Private/DaySequenceEditorModule.cpp
--- a/Engine/Plugins/Experimental/DaySequence/Source/DaySequenceEditor/Private/DaySequenceEditorModule.cpp
+++ b/Engine/Plugins/Experimental/DaySequence/Source/DaySequenceEditor/Private/DaySequenceEditorModule.cpp
@@ -23,6 +23,7 @@
 #include "ISequencerModule.h"
 #include "ISettingsModule.h"
 #include "SequencerSettings.h"
+#include "Sections/MovieSceneSubSection.h"
 
 #include "Application/ThrottleManager.h"
 #include "Framework/Docking/TabManager.h"
@@ -137,6 +138,8 @@ private:
 
 	FDelegateHandle OnSwitchPIEAndSIEHandle;
 
+	FDelegateHandle OnSubSectionRemovedHandle;
+
 	int32 LastAllowThrottling = 1;
 };
 
@@ -165,6 +168,35 @@ void FDaySequenceEditorModule::StartupModule()
 	{
 		UDaySequenceModifierComponent::SetIsSimulating(bIsSimulating);
 	});
+
+	OnSubSectionRemovedHandle = ADaySequenceActor::OnSubSectionRemovedEvent.AddLambda([](const UMovieSceneSubSection* RemovedSubSection)
+	{
+		FDaySequenceEditorToolkit::IterateOpenToolkits([RemovedSubSection](const FDaySequenceEditorToolkit& Toolkit)
+		{
+			if (!Toolkit.IsActorPreview())
+			{
+				// In rare cases a focused subsection can be removed right as an evaluation is triggered,
+				// breaking assumptions in core Sequencer code that the focused subsequence is always valid.
+				// Any time a subsection is removed we should enforce this assumption.
+				if (const TSharedPtr<ISequencer> ToolkitSequencer = Toolkit.GetSequencer())
+				{
+					const UMovieSceneSequence* FocusedSequence = ToolkitSequencer->GetFocusedMovieSceneSequence();
+					const UMovieSceneSequence* RemovedSequence = RemovedSubSection ? RemovedSubSection->GetSequence() : nullptr;
+					
+					if (!FocusedSequence || FocusedSequence == RemovedSequence)
+					{
+						ToolkitSequencer->PopToSequenceInstance(ToolkitSequencer->GetRootTemplateID());
+					}
+				}
+				
+				// Break out of iteration now that we found a sequence editor
+				return false;
+			}
+
+			// Continue looking for sequence editor toolkit
+			return true;
+		});
+	});
 	
 	PluginCommands = MakeShareable(new FUICommandList);
 
@@ -261,6 +293,8 @@ void FDaySequenceEditorModule::ShutdownModule()
 
 	FEditorDelegates::OnSwitchBeginPIEAndSIE.Remove(OnSwitchPIEAndSIEHandle);
 
+	ADaySequenceActor::OnSubSectionRemovedEvent.Remove(OnSubSectionRemovedHandle);
+
 	if (FPropertyEditorModule* PropertyModule = FModuleManager::GetModulePtr<FPropertyEditorModule>(PropertyEditorModuleName))
 	{
 		PropertyModule->UnregisterCustomClassLayout(DaySequenceActorClassName);
