DaySequenceModifier: Don't inject subsections on server.

diff --git a/Engine/Plugins/Experimental/DaySequence/Source/DaySequence/Private/DaySequenceModifierComponent.cpp b/Engine/Plugins/Experimental/DaySequence/Source/DaySequence/Private/DaySequenceModifierComponent.cpp
--- a/Engine/Plugins/Experimental/DaySequence/Source/DaySequence/Private/DaySequenceModifierComponent.cpp
+++ b/Engine/Plugins/Experimental/DaySequence/Source/DaySequence/Private/DaySequenceModifierComponent.cpp
@@ -692,6 +692,11 @@ UMovieSceneSubSection* UDaySequenceModifierComponent::InitializeDaySequence(cons
     {
     	return nullptr;
     }
+
+	if (const UWorld* World = GetWorld(); World && World->WorldType != EWorldType::Editor && TargetActor->GetNetMode() == NM_DedicatedServer)
+	{
+		return nullptr;
+	}
 	
 	auto CreateSubTrack = [this, MovieScene](UDaySequence* Sequence, int BiasOffset, bool bActivate, bool bBlendHierarchicalBias)
 	{
