DaySequenceModifierVolume: Add support for replay scrubbing

diff --git a/Engine/Plugins/Experimental/DaySequence/Source/DaySequence/Private/Actors/DaySequenceModifierVolume.cpp b/Engine/Plugins/Experimental/DaySequence/Source/DaySequence/Private/Actors/DaySequenceModifierVolume.cpp
--- a/Engine/Plugins/Experimental/DaySequence/Source/DaySequence/Private/Actors/DaySequenceModifierVolume.cpp
+++ b/Engine/Plugins/Experimental/DaySequence/Source/DaySequence/Private/Actors/DaySequenceModifierVolume.cpp
@@ -11,6 +11,7 @@
 #include "GameDelegates.h"
 #include "GameFramework/Pawn.h"
 #include "TimerManager.h"
+#include "Net/UnrealNetwork.h"
 
 ADaySequenceModifierVolume::ADaySequenceModifierVolume(const FObjectInitializer& Init)
 : Super(Init)
@@ -39,11 +40,22 @@ void ADaySequenceModifierVolume::SetBlendTarget(APlayerController* InPC)
 	DaySequenceModifier->SetUserBlendWeight(1.f);
 }
 
-void ADaySequenceModifierVolume::BeginPlay()
+void ADaySequenceModifierVolume::BeginPlay()	
 {
 	Super::BeginPlay();
 
 	Initialize();
+
+	if (const UWorld* World = GetWorld(); World && World->IsPlayingReplay())
+	{
+		ReplayScrubbedHandle = FNetworkReplayDelegates::OnReplayScrubComplete.AddWeakLambda(this, [this](const UWorld* InWorld)
+		{
+			if (InWorld == GetWorld())
+			{
+				DaySequenceActorSetup();
+			}
+		});
+	}
 }
 
 void ADaySequenceModifierVolume::OnConstruction(const FTransform& Transform)
diff --git a/Engine/Plugins/Experimental/DaySequence/Source/DaySequence/Public/Actors/DaySequenceModifierVolume.h b/Engine/Plugins/Experimental/DaySequence/Source/DaySequence/Public/Actors/DaySequenceModifierVolume.h
--- a/Engine/Plugins/Experimental/DaySequence/Source/DaySequence/Public/Actors/DaySequenceModifierVolume.h
+++ b/Engine/Plugins/Experimental/DaySequence/Source/DaySequence/Public/Actors/DaySequenceModifierVolume.h
@@ -93,4 +93,6 @@ protected:
 	TObjectPtr<AActor> CurrentBlendTarget;
 	
 	FDelegateHandle ViewTargetChangedHandle;
+	
+	FDelegateHandle ReplayScrubbedHandle;
 };
