DaySequenceModifierVolume: Fix configuration of shape components

diff --git a/Engine/Plugins/Experimental/DaySequence/Source/DaySequence/Private/Actors/DaySequenceModifierVolume.cpp b/Engine/Plugins/Experimental/DaySequence/Source/DaySequence/Private/Actors/DaySequenceModifierVolume.cpp
--- a/Engine/Plugins/Experimental/DaySequence/Source/DaySequence/Private/Actors/DaySequenceModifierVolume.cpp
+++ b/Engine/Plugins/Experimental/DaySequence/Source/DaySequence/Private/Actors/DaySequenceModifierVolume.cpp
@@ -25,6 +25,10 @@ ADaySequenceModifierVolume::ADaySequenceModifierVolume(const FObjectInitializer&
 	DefaultBox->SetupAttachment(DaySequenceModifier);
 	DefaultBox->SetLineThickness(10.f);
 	DefaultBox->SetBoxExtent(FVector(500.f));
+
+	FComponentReference DefaultBoxReference;
+	DefaultBoxReference.ComponentProperty = TEXT("DefaultBox");
+	DaySequenceModifier->AddVolumeShapeComponent(DefaultBoxReference);
 }
 
 void ADaySequenceModifierVolume::SetBlendTarget(APlayerController* InPC)
@@ -71,11 +75,6 @@ void ADaySequenceModifierVolume::Initialize()
 	{
 		return;
 	}
-
-	FComponentReference DefaultBoxReference;
-	DefaultBoxReference.OverrideComponent = DefaultBox;
-	DaySequenceModifier->EmptyVolumeShapeComponents();
-	DaySequenceModifier->AddVolumeShapeComponent(DefaultBoxReference);
 	
 	// This actor should only initialize on the client.
 	if (GetNetMode() == NM_DedicatedServer)
diff --git a/Engine/Plugins/Experimental/DaySequence/Source/DaySequence/Private/DaySequenceModifierComponent.cpp b/Engine/Plugins/Experimental/DaySequence/Source/DaySequence/Private/DaySequenceModifierComponent.cpp
--- a/Engine/Plugins/Experimental/DaySequence/Source/DaySequence/Private/DaySequenceModifierComponent.cpp
+++ b/Engine/Plugins/Experimental/DaySequence/Source/DaySequence/Private/DaySequenceModifierComponent.cpp
@@ -1013,7 +1013,7 @@ void UDaySequenceModifierComponent::EmptyVolumeShapeComponents()
 
 void UDaySequenceModifierComponent::AddVolumeShapeComponent(const FComponentReference& InShapeReference)
 {
-	VolumeShapeComponents.Add(InShapeReference);
+	VolumeShapeComponents.AddUnique(InShapeReference);
 	bCachedExternalShapesInvalid = true;
 }
 
