DaySequence: Fix Override Initial Time of Day

Override Run Day Cycle seems to work fine, the problem with Override Initial Time of Day was that TimeOfDayPreview was getting overwritten by InitialTimeOfDay on load

diff --git a/Engine/Plugins/Experimental/DaySequence/Source/DaySequence/Private/DaySequenceActor.cpp b/Engine/Plugins/Experimental/DaySequence/Source/DaySequence/Private/DaySequenceActor.cpp
--- a/Engine/Plugins/Experimental/DaySequence/Source/DaySequence/Private/DaySequenceActor.cpp
+++ b/Engine/Plugins/Experimental/DaySequence/Source/DaySequence/Private/DaySequenceActor.cpp
@@ -253,7 +253,11 @@ void ADaySequenceActor::PostLoad()
 
 #if WITH_EDITOR
 	// Initialize our transient preview time to InitialTimeOfDay on load.
-	TimeOfDayPreview = InitialTimeOfDay;
+	// Only do this for editor world, in PIE world we want to preserve the value in case we are overriding initial time of day.
+	if (const UWorld* World = GetWorld(); World && World->WorldType == EWorldType::Editor)
+	{
+		TimeOfDayPreview = InitialTimeOfDay;
+	}
 	
 	// Build our root sequence after load to ensure that the editor can parse the root sequence
 	// hierarchy for editing binding overrides. This is only necessary for editor, since the
