DaySequenceModifierComponent: Use Setter/Getter for sequence properties. Also fix issue with busy setting from sequencer.

We need sequencer to use accessors for this property to correctly set it during game.

diff --git a/Engine/Plugins/Experimental/DaySequence/Source/DaySequence/Private/DaySequenceModifierComponent.cpp b/Engine/Plugins/Experimental/DaySequence/Source/DaySequence/Private/DaySequenceModifierComponent.cpp
--- a/Engine/Plugins/Experimental/DaySequence/Source/DaySequence/Private/DaySequenceModifierComponent.cpp
+++ b/Engine/Plugins/Experimental/DaySequence/Source/DaySequence/Private/DaySequenceModifierComponent.cpp
@@ -830,11 +830,17 @@ void UDaySequenceModifierComponent::SetUserDaySequence(UDaySequence* InDaySequen
 		return;
 	}
 
+	// Prevents unnecessary & expensive subsequence reinitialization
+	if (InDaySequence == UserDaySequence)
+	{
+		return;
+	}
+
 	UserDaySequence = InDaySequence;
 	ReinitializeSubSequence(nullptr);
 }
 
-UDaySequence* UDaySequenceModifierComponent::GetUserDaySequence()
+UDaySequence* UDaySequenceModifierComponent::GetUserDaySequence() const
 {
 	return UserDaySequence;
 }
@@ -846,12 +852,18 @@ void UDaySequenceModifierComponent::SetTransientSequence(UDaySequence* InDaySequ
 		FFrame::KismetExecutionMessage(TEXT("SetTransientSequence called with a non-transient sequence, use SetUserDaySequence instead!"), ELogVerbosity::Error);
 		return;
 	}
+
+	// Prevents unnecessary & expensive subsequence reinitialization
+	if (InDaySequence == TransientSequence)
+	{
+		return;
+	}
 	
 	TransientSequence = InDaySequence;
 	ReinitializeSubSequence(nullptr);
 }
 	
-UDaySequence* UDaySequenceModifierComponent::GetTransientSequence()
+UDaySequence* UDaySequenceModifierComponent::GetTransientSequence() const
 {
 	return TransientSequence;
 }
diff --git a/Engine/Plugins/Experimental/DaySequence/Source/DaySequence/Public/DaySequenceModifierComponent.h b/Engine/Plugins/Experimental/DaySequence/Source/DaySequence/Public/DaySequenceModifierComponent.h
--- a/Engine/Plugins/Experimental/DaySequence/Source/DaySequence/Public/DaySequenceModifierComponent.h
+++ b/Engine/Plugins/Experimental/DaySequence/Source/DaySequence/Public/DaySequenceModifierComponent.h
@@ -166,18 +166,12 @@ public:
 	DAYSEQUENCE_API void DisableComponent();
 
 	/** Sets the user day sequence. This must be a user created asset. */
-	UFUNCTION(BlueprintSetter, Category="Day Sequence")
 	DAYSEQUENCE_API void SetUserDaySequence(UDaySequence* InDaySequence);
-
-	UFUNCTION(BlueprintGetter, BlueprintPure, Category="Day Sequence")
-	DAYSEQUENCE_API UDaySequence* GetUserDaySequence();
+	DAYSEQUENCE_API UDaySequence* GetUserDaySequence() const;
 
 	/** Sets a user specified transient sequence. */
-	UFUNCTION(BlueprintSetter, Category="Day Sequence")
 	DAYSEQUENCE_API void SetTransientSequence(UDaySequence* InDaySequence);
-	
-	UFUNCTION(BlueprintGetter, BlueprintPure, Category="Day Sequence")
-	DAYSEQUENCE_API UDaySequence* GetTransientSequence();
+	DAYSEQUENCE_API UDaySequence* GetTransientSequence() const;
 	
 	UFUNCTION(BlueprintCallable, Category="Day Sequence")
 	DAYSEQUENCE_API void SetDayNightCycle(EDayNightCycleMode NewMode);
@@ -261,11 +255,11 @@ protected:
 	TWeakObjectPtr<APlayerController> WeakBlendTarget;
 
 	/** The user provided Day Sequence. */
-	UPROPERTY(EditAnywhere, Category="Day Sequence", BlueprintGetter = GetUserDaySequence, BlueprintSetter = SetUserDaySequence, meta=(EditCondition="!bUseCollection", EditConditionHides, DisplayAfter="bUseCollection"))
+	UPROPERTY(EditAnywhere, BlueprintReadWrite, Category="Day Sequence", Setter, Getter, meta=(EditCondition="!bUseCollection", EditConditionHides, DisplayAfter="bUseCollection"))
 	TObjectPtr<UDaySequence> UserDaySequence;
 
 	/** The user provided Transient Day Sequence. */
-    UPROPERTY(Transient, Category="Day Sequence", BlueprintGetter = GetTransientSequence, BlueprintSetter = SetTransientSequence, meta=(EditCondition="!bUseCollection", EditConditionHides, DisplayAfter="bUseCollection"))
+    UPROPERTY(Transient, BlueprintReadWrite, Category="Day Sequence", Setter, Getter, meta=(EditCondition="!bUseCollection", EditConditionHides, DisplayAfter="bUseCollection"))
     TObjectPtr<UDaySequence> TransientSequence;
 	
 	/** The user provided collection. This is an alternative to UserDaySequence. */
