Iris
* Fix some experimental code not compiling when UE_WITH_IRIS is false

#rnx

diff --git a/Engine/Plugins/Experimental/DaySequence/Source/DaySequence/Private/DaySequencePlayer.cpp b/Engine/Plugins/Experimental/DaySequence/Source/DaySequence/Private/DaySequencePlayer.cpp
--- a/Engine/Plugins/Experimental/DaySequence/Source/DaySequence/Private/DaySequencePlayer.cpp
+++ b/Engine/Plugins/Experimental/DaySequence/Source/DaySequence/Private/DaySequencePlayer.cpp
@@ -22,7 +22,9 @@
 #include "GameFramework/WorldSettings.h"
 #include "GameFramework/PlayerController.h"
 #include "GameFramework/PlayerState.h"
+#if UE_WITH_IRIS
 #include "Iris/ReplicationSystem/ReplicationFragmentUtil.h"
+#endif
 #include "Net/UnrealNetwork.h"
 #include "Engine/NetDriver.h"
 
