DaySequenceActor: Undo/redo handling

Reinit the root sequence when an unspecified property change occurs.

diff --git a/Engine/Plugins/Experimental/DaySequence/Source/DaySequence/Private/DaySequenceActor.cpp b/Engine/Plugins/Experimental/DaySequence/Source/DaySequence/Private/DaySequenceActor.cpp
--- a/Engine/Plugins/Experimental/DaySequence/Source/DaySequence/Private/DaySequenceActor.cpp
+++ b/Engine/Plugins/Experimental/DaySequence/Source/DaySequence/Private/DaySequenceActor.cpp
@@ -405,6 +405,12 @@ bool ADaySequenceActor::GetReferencedContentObjects(TArray<UObject*>& Objects) c
 
 void ADaySequenceActor::PostEditChangeProperty(FPropertyChangedEvent& PropertyChangedEvent)
 {
+	auto ReinitializeActor = [this]()
+	{
+		bUpdateRootSequenceOnTick = true;
+		SubSections.Empty();
+	};
+	
 	const FName PropertyName = (PropertyChangedEvent.Property != nullptr) ? PropertyChangedEvent.Property->GetFName() : NAME_None;
 	if (PropertyName == GET_MEMBER_NAME_CHECKED(ADaySequenceActor, TimeOfDayPreview))
 	{
@@ -419,11 +425,9 @@ void ADaySequenceActor::PostEditChangeProperty(FPropertyChangedEvent& PropertyCh
 	{
 		SetTimePerCycle(GetTimePerCycle());
 
+		// need to null this out to guarantee total reconstruction.
 		RootSequence = nullptr;
-		SubSections.Empty();
-		
-		// Update our root sequence with the new asset.
-		UpdateRootSequence();
+		ReinitializeActor();
 	}
 	else if (PropertyName == GET_MEMBER_NAME_CHECKED(ADaySequenceActor, InitialTimeOfDay))
 	{
@@ -431,7 +435,12 @@ void ADaySequenceActor::PostEditChangeProperty(FPropertyChangedEvent& PropertyCh
 	}
 	else if (PropertyName == GET_MEMBER_NAME_CHECKED(ADaySequenceActor, DaySequenceCollection))
 	{
-		SubSections.Empty();
+		ReinitializeActor();
+	}
+	else if (PropertyChangedEvent.ChangeType == EPropertyChangeType::Unspecified)
+	{
+		// This handles undo/redo transactions.
+		ReinitializeActor();
 	}
 
 	Super::PostEditChangeProperty(PropertyChangedEvent);
diff --git a/Engine/Plugins/Experimental/DaySequence/Source/DaySequence/Public/DaySequenceActor.h b/Engine/Plugins/Experimental/DaySequence/Source/DaySequence/Public/DaySequenceActor.h
--- a/Engine/Plugins/Experimental/DaySequence/Source/DaySequence/Public/DaySequenceActor.h
+++ b/Engine/Plugins/Experimental/DaySequence/Source/DaySequence/Public/DaySequenceActor.h
@@ -333,6 +333,9 @@ protected:
 	//~ Begin UObject interface
 	virtual void PostInitProperties() override;
 	virtual void PostLoad() override;
+#if WITH_EDITOR
+	virtual void PostEditChangeProperty(FPropertyChangedEvent& PropertyChangedEvent) override;
+#endif
 	//~ End UObject interface
 	
 	//~ Begin AActor interface
@@ -366,8 +369,6 @@ protected:
 	virtual void UpdateObjectFromProxy(FStructOnScope& Proxy, IPropertyHandle& ObjectPropertyHandle) override;
 	virtual UMovieSceneSequence* RetrieveOwnedSequence() const override;
 	//~ End IMovieSceneBindingOwnerInterface
-	
-	virtual void PostEditChangeProperty(FPropertyChangedEvent& PropertyChangedEvent) override;
 #endif //WITH_EDITOR
 
 	/** Initialize SequencePlayer with transient root sequence */
