DaySequenceModifierComponent: Add BP accessors to several properties

A number of our properties require non-trivial accessing logic (at least when setting, not always when getting). For example, setting DayNightCycle requires the modifier component to reenable and setting UserBlendWeight requires the set value be clamped.

Importantly note: we don't need to add accessors to properties that didn't already have them, we just need to make sure no accessors are BP callable. We can always introduce and remove accessors as long as they are invoked by BPs that are attempting to directly set the property because of UHT magic.

Another important note: I have observed that Setter/Getter markup is not respected by the details panel. I think most other things do respect it. For now I am just handling this in PostEditChangeProperty by routing the new value into the setter for properties that need this behavior. Trivial setters don't need this treatment.

This change also corrects the DayNightCycle setter to properly reenable the modifier on set.

#changelist validated
#virtualized

diff --git a/Engine/Plugins/Experimental/DaySequence/Source/DaySequence/Private/DaySequenceModifierComponent.cpp b/Engine/Plugins/Experimental/DaySequence/Source/DaySequence/Private/DaySequenceModifierComponent.cpp
--- a/Engine/Plugins/Experimental/DaySequence/Source/DaySequence/Private/DaySequenceModifierComponent.cpp
+++ b/Engine/Plugins/Experimental/DaySequence/Source/DaySequence/Private/DaySequenceModifierComponent.cpp
@@ -248,6 +248,11 @@ void UDaySequenceModifierComponent::PostEditChangeProperty(FPropertyChangedEvent
 	{
 		EnableModifier();
 	}
+	else if (PropertyName == GET_MEMBER_NAME_CHECKED(UDaySequenceModifierComponent, DayNightCycle))
+	{
+		// Force details panel changes to use our setter.
+		SetDayNightCycle(DayNightCycle);
+	}
 }
 
 #endif // WITH_EDITOR
@@ -509,13 +514,8 @@ void UDaySequenceModifierComponent::DisableModifier()
 
 		TargetActor->UnregisterStaticTimeContributor(this);
 		
-		if (bUnpauseOnDisable)
-		{
-			TargetActor->Play();
-			bUnpauseOnDisable = false;
-		}
 		// Force an update if it's not playing so that the effects of this being disabled are seen
-		else if (!TargetActor->IsPlaying())
+		if (!TargetActor->IsPlaying())
 		{
 			TargetActor->SetTimeOfDay(TargetActor->GetTimeOfDay());
 		}
@@ -871,6 +871,17 @@ UDaySequence* UDaySequenceModifierComponent::GetTransientSequence() const
 void UDaySequenceModifierComponent::SetDayNightCycle(EDayNightCycleMode NewMode)
 {
 	DayNightCycle = NewMode;
+
+	if (bIsComponentEnabled)
+	{
+		DisableComponent();
+		EnableComponent();
+	}
+}
+
+EDayNightCycleMode UDaySequenceModifierComponent::GetDayNightCycle() const
+{
+	return DayNightCycle;
 }
 
 void UDaySequenceModifierComponent::SetDayNightCycleTime(float Time)
@@ -878,16 +889,46 @@ void UDaySequenceModifierComponent::SetDayNightCycleTime(float Time)
 	DayNightCycleTime = Time;
 }
 
+float UDaySequenceModifierComponent::GetDayNightCycleTime() const
+{
+	return DayNightCycleTime;
+}
+
 void UDaySequenceModifierComponent::SetMode(EDaySequenceModifierMode NewMode)
 {
 	Mode = NewMode;
 }
 
+EDaySequenceModifierMode UDaySequenceModifierComponent::GetMode() const
+{
+	return Mode;
+}
+
 void UDaySequenceModifierComponent::SetBlendPolicy(EDaySequenceModifierUserBlendPolicy NewPolicy)
 {
 	BlendPolicy = NewPolicy;
 }
 
+EDaySequenceModifierUserBlendPolicy UDaySequenceModifierComponent::GetBlendPolicy() const
+{
+	return BlendPolicy;
+}
+
+void UDaySequenceModifierComponent::SetBlendTarget(APlayerController* InActor)
+{	
+	WeakBlendTarget = InActor;
+}
+
+void UDaySequenceModifierComponent::SetUserBlendWeight(float Weight)
+{
+	UserBlendWeight = FMath::Clamp(Weight, 0.f, 1.f);
+}
+
+float UDaySequenceModifierComponent::GetUserBlendWeight() const
+{
+	return UserBlendWeight;
+}
+
 bool UDaySequenceModifierComponent::GetBlendPosition(FVector& InPosition) const
 {
 	CSV_SCOPED_TIMING_STAT(DaySequence, GetBlendPosition);
@@ -1034,16 +1075,6 @@ void UDaySequenceModifierComponent::InvalidateMuteStates() const
 	OnInvalidateMuteStates.Broadcast();
 }
 
-void UDaySequenceModifierComponent::SetBlendTarget(APlayerController* InActor)
-{	
-	WeakBlendTarget = InActor;
-}
-
-void UDaySequenceModifierComponent::SetUserBlendWeight(float Weight)
-{
-	UserBlendWeight = FMath::Clamp(Weight, 0.f, 1.f);
-}
-
 #if ENABLE_DRAW_DEBUG
 void UDaySequenceModifierComponent::OnDebugLevelChanged(int32 InDebugLevel)
 {
diff --git a/Engine/Plugins/Experimental/DaySequence/Source/DaySequence/Public/DaySequenceModifierComponent.h b/Engine/Plugins/Experimental/DaySequence/Source/DaySequence/Public/DaySequenceModifierComponent.h
--- a/Engine/Plugins/Experimental/DaySequence/Source/DaySequence/Public/DaySequenceModifierComponent.h
+++ b/Engine/Plugins/Experimental/DaySequence/Source/DaySequence/Public/DaySequenceModifierComponent.h
@@ -173,26 +173,27 @@ public:
 	DAYSEQUENCE_API void SetTransientSequence(UDaySequence* InDaySequence);
 	DAYSEQUENCE_API UDaySequence* GetTransientSequence() const;
 	
-	UFUNCTION(BlueprintCallable, Category="Day Sequence")
+	/** Sets the Day Night Cycle mode. This will reenable the component. */
 	DAYSEQUENCE_API void SetDayNightCycle(EDayNightCycleMode NewMode);
+	DAYSEQUENCE_API EDayNightCycleMode GetDayNightCycle() const;
 
-	UFUNCTION(BlueprintCallable, Category="Day Sequence")
 	DAYSEQUENCE_API void SetDayNightCycleTime(float Time);
+	DAYSEQUENCE_API float GetDayNightCycleTime() const;
 	
-	UFUNCTION(BlueprintCallable, Category="Day Sequence")
 	DAYSEQUENCE_API void SetMode(EDaySequenceModifierMode NewMode);
+	DAYSEQUENCE_API EDaySequenceModifierMode GetMode() const;
 
-	UFUNCTION(BlueprintCallable, Category="Day Sequence")
 	DAYSEQUENCE_API void SetBlendPolicy(EDaySequenceModifierUserBlendPolicy NewPolicy);
+	DAYSEQUENCE_API EDaySequenceModifierUserBlendPolicy GetBlendPolicy() const;
+
+	/** Sets a custom blend weight for volume based blends. Final weight depends on BlendPolicy. */
+	DAYSEQUENCE_API void SetUserBlendWeight(float Weight);
+	DAYSEQUENCE_API float GetUserBlendWeight() const;
 
 	/** Sets the blend target to use when in Volume mode. */
 	UFUNCTION(BlueprintCallable, Category="Day Sequence")
 	DAYSEQUENCE_API void SetBlendTarget(APlayerController* InActor);
 	
-	/** Sets a custom blend weight for volume based blends. Final weight depends on BlendPolicy. */
-	UFUNCTION(BlueprintCallable, Category="Day Sequence")
-	DAYSEQUENCE_API void SetUserBlendWeight(float Weight);
-
 	/** Get the current blend weight. */
 	UFUNCTION(BlueprintCallable, Category="Day Sequence")
 	float GetBlendWeight() const;
@@ -255,11 +256,11 @@ protected:
 	TWeakObjectPtr<APlayerController> WeakBlendTarget;
 
 	/** The user provided Day Sequence. */
-	UPROPERTY(EditAnywhere, BlueprintReadWrite, Category="Day Sequence", Setter, Getter, meta=(EditCondition="!bUseCollection", EditConditionHides, DisplayAfter="bUseCollection"))
+	UPROPERTY(EditAnywhere, BlueprintReadWrite, Setter, Getter, Category="Day Sequence", meta=(EditCondition="!bUseCollection", EditConditionHides, DisplayAfter="bUseCollection"))
 	TObjectPtr<UDaySequence> UserDaySequence;
 
 	/** The user provided Transient Day Sequence. */
-    UPROPERTY(Transient, BlueprintReadWrite, Category="Day Sequence", Setter, Getter, meta=(EditCondition="!bUseCollection", EditConditionHides, DisplayAfter="bUseCollection"))
+    UPROPERTY(Transient, BlueprintReadWrite, Setter, Getter, Category="Day Sequence", meta=(EditCondition="!bUseCollection", EditConditionHides, DisplayAfter="bUseCollection"))
     TObjectPtr<UDaySequence> TransientSequence;
 	
 	/** The user provided collection. This is an alternative to UserDaySequence. */
@@ -274,27 +275,27 @@ protected:
 	int32 Bias;
 
 	/** The time to use for the day/night cycle. */
-	UPROPERTY(BlueprintReadWrite, EditAnywhere, Category="Time", DisplayName="Time", meta=(DisplayAfter="DayNightCycle", EditCondition="DayNightCycle==EDayNightCycleMode::FixedTime || DayNightCycle==EDayNightCycleMode::StartAtSpecifiedTime", EditConditionHides))
+	UPROPERTY(EditAnywhere, BlueprintReadWrite, Setter, Getter, DisplayName="Time", Category="Time", meta=(DisplayAfter="DayNightCycle", EditCondition="DayNightCycle==EDayNightCycleMode::FixedTime || DayNightCycle==EDayNightCycleMode::StartAtSpecifiedTime", EditConditionHides))
 	float DayNightCycleTime;
 
 	/** Defines the region in which the effective blend weight is in the range (0.0, 1.0) (not inclusive) when Mode == EDaySequenceModifierMode::Volume. */
-	UPROPERTY(BlueprintReadWrite, EditAnywhere, Category="Day Sequence", meta=(DisplayAfter="Mode", EditCondition="Mode==EDaySequenceModifierMode::Volume", EditConditionHides))
+	UPROPERTY(EditAnywhere, BlueprintReadWrite, Category="Day Sequence", meta=(DisplayAfter="Mode", EditCondition="Mode==EDaySequenceModifierMode::Volume", EditConditionHides))
 	float BlendAmount;
 
 	/** User specified blend weight. The final blend weight is determined by BlendPolicy. */
-	UPROPERTY(Transient, BlueprintReadWrite, EditAnywhere, Category="Day Sequence", meta=(DisplayAfter="BlendPolicy", EditCondition="BlendPolicy!=EDaySequenceModifierUserBlendPolicy::Ignored", EditConditionHides))
+	UPROPERTY(Transient, EditAnywhere, BlueprintReadWrite, Setter, Getter, Category="Day Sequence", meta=(DisplayAfter="BlendPolicy", EditCondition="BlendPolicy!=EDaySequenceModifierUserBlendPolicy::Ignored", EditConditionHides))
 	float UserBlendWeight;
 
 	/** Changes the way the modifier controls the day/night cycle time when enabled. */
-	UPROPERTY(BlueprintReadWrite, EditAnywhere, Category="Time", DisplayName="Day/Night Cycle")
+	UPROPERTY(EditAnywhere, BlueprintReadWrite, Setter, Getter, DisplayName="Day/Night Cycle", Category="Time")
 	EDayNightCycleMode DayNightCycle;
 
 	/** Determines how the modifier computes InternalBlendWeight. */
-	UPROPERTY(BlueprintReadWrite, EditAnywhere, Category="Day Sequence")
+	UPROPERTY(EditAnywhere, BlueprintReadWrite, Setter, Getter, Category="Day Sequence")
 	EDaySequenceModifierMode Mode;
 
 	/** Determines how the modifier uses UserBlendWeight to compute effective blend weight. */
-	UPROPERTY(BlueprintReadWrite, EditAnywhere, Category="Day Sequence")
+	UPROPERTY(EditAnywhere, BlueprintReadWrite, Setter, Getter, Category="Day Sequence")
 	EDaySequenceModifierUserBlendPolicy BlendPolicy;
 
 	/** Blueprint exposed delegate invoked after the component's subsequences are reinitialized. */
@@ -321,15 +322,12 @@ protected:
 	UPROPERTY(Transient, BlueprintReadWrite, EditAnywhere, Category="Day Sequence")
 	uint8 bPreview : 1;
 
-	/** Flag to keep track of whether we need to unpause the day sequence when we are disabled */
-	uint8 bUnpauseOnDisable : 1;
-
 	/** If true, hide UserDaySequence and expose DaySequenceCollection. */
 	UPROPERTY(EditAnywhere, Category="Day Sequence")
 	uint8 bUseCollection : 1;
 
 	/** If true, day sequence evaluation while within the blending region will be smooth. Note: Can be very expensive. */
-	UPROPERTY(EditAnywhere, BlueprintReadWrite, Category="Day Sequence")
+	UPROPERTY(EditAnywhere, BlueprintReadWrite, AdvancedDisplay, Category="Day Sequence")
 	uint8 bSmoothBlending : 1;
 
 private:
