DaySequence: Fix up some debug overlay problems

Debug drawing was getting disabled on some netmodes because of a faulty HasAuthority check. We only want to disable on dedicated servers. Also, we weren't actually checking ShowCondition.

diff --git a/Engine/Plugins/Experimental/DaySequence/Source/DaySequence/Private/DaySequenceActor.cpp b/Engine/Plugins/Experimental/DaySequence/Source/DaySequence/Private/DaySequenceActor.cpp
--- a/Engine/Plugins/Experimental/DaySequence/Source/DaySequence/Private/DaySequenceActor.cpp
+++ b/Engine/Plugins/Experimental/DaySequence/Source/DaySequence/Private/DaySequenceActor.cpp
@@ -1534,7 +1534,7 @@ void ADaySequenceActor::OnShowDebugInfo(AHUD* HUD, UCanvas* Canvas, const FDebug
 {
 	using namespace UE::DaySequence;
 	
-	if (!Canvas || !GEngine || HasAuthority())
+	if (!Canvas || !GEngine || GetNetMode() == NM_DedicatedServer)
 	{
 		return;
 	}
@@ -1543,19 +1543,18 @@ void ADaySequenceActor::OnShowDebugInfo(AHUD* HUD, UCanvas* Canvas, const FDebug
 		return !Category.Value.Key.IsEmpty();
 	});
 
-	auto RemoveStaleEntriesAndGetPinnedArray = [](FDebugEntryArray& WeakArray, TArray<TSharedPtr<TMap<FString, FString>>>& OutSharedArray)
+	auto RemoveStaleEntriesAndGetPinnedArray = [](FDebugEntryArray& InWeakArray, TArray<TSharedPtr<TMap<FString, FString>>>& OutSharedArray)
 	{
 		// Shouldn't reduce existing capacity, so in theory the total number
 		// of allocations here will be <= the size of the largest array in DebugEntries
 		OutSharedArray.Reset();
 		
-		WeakArray.RemoveAll([](const TWeakPtr<FDaySequenceDebugEntry>& Entry) 
-			{ return !Entry.IsValid(); });
+		auto RemovePredicate = [](const TWeakPtr<FDaySequenceDebugEntry>& Entry){ return !Entry.IsValid(); };
+		InWeakArray.RemoveAll(RemovePredicate);
 
-		for (TWeakPtr<FDaySequenceDebugEntry> Entry : WeakArray)
-		{
-			OutSharedArray.Push(Entry.Pin()->GetData());
-		}
+		auto TransformPredicate = [](const TWeakPtr<FDaySequenceDebugEntry>& Entry){ return Entry.Pin()->ShowCondition(); };
+		auto TransformFunction = [](const TWeakPtr<FDaySequenceDebugEntry>& Entry){ return Entry.Pin()->GetData(); };
+		Algo::TransformIf(InWeakArray, OutSharedArray, TransformPredicate, TransformFunction);
 	};
 	
 	TArray<TSharedPtr<TMap<FString, FString>>> EntriesToDraw;
diff --git a/Engine/Plugins/Experimental/DaySequence/Source/DaySequence/Private/DaySequenceModifierComponent.cpp b/Engine/Plugins/Experimental/DaySequence/Source/DaySequence/Private/DaySequenceModifierComponent.cpp
--- a/Engine/Plugins/Experimental/DaySequence/Source/DaySequence/Private/DaySequenceModifierComponent.cpp
+++ b/Engine/Plugins/Experimental/DaySequence/Source/DaySequence/Private/DaySequenceModifierComponent.cpp
@@ -1083,11 +1083,6 @@ void UDaySequenceModifierComponent::OnDebugLevelChanged(int32 InDebugLevel)
 
 bool UDaySequenceModifierComponent::ShouldShowDebugInfo() const
 {
-	if (GetOwner()->HasAuthority())
-	{
-		return false;
-	}
-	
 	switch (DebugLevel)
 	{
 	case 0: return false;
