DaySequence: Fixup several more debug overlay issues.

This system is actually more confusing and fragile than I thought when I built it, especially when testing with multiple local clients in editor. We should now be properly respecting debug level in modifier components, skipping drawing in dedicated server actors, and no longer drawing debug info from the editor world actor.

diff --git a/Engine/Plugins/Experimental/DaySequence/Source/DaySequence/Private/DaySequenceActor.cpp b/Engine/Plugins/Experimental/DaySequence/Source/DaySequence/Private/DaySequenceActor.cpp
--- a/Engine/Plugins/Experimental/DaySequence/Source/DaySequence/Private/DaySequenceActor.cpp
+++ b/Engine/Plugins/Experimental/DaySequence/Source/DaySequence/Private/DaySequenceActor.cpp
@@ -24,7 +24,7 @@
 #include "UObject/UObjectThreadContext.h" // For FUObjectThreadContext
 #include "ProfilingDebugging/CsvProfiler.h"
 
-#if ENABLE_DRAW_DEBUG
+#if DAY_SEQUENCE_ENABLE_DRAW_DEBUG
 #include "Engine/Engine.h"	// GEngine
 #include "GameFramework/HUD.h"	// AHUD::OnShowDebugInfo
 #include "Engine/Canvas.h"	// UCanvas
@@ -46,13 +46,13 @@ ADaySequenceActor::FOnSubSectionRemovedEvent ADaySequenceActor::OnSubSectionRemo
 
 namespace UE::DaySequence
 {
-#if ENABLE_DRAW_DEBUG
+#if DAY_SEQUENCE_ENABLE_DRAW_DEBUG
 	FDaySequenceDebugEntry::FDaySequenceDebugEntry(FShowDebugDataConditionFunction InShowCondition, FGetDebugDataFunction InGetData)
 	: ShowCondition(InShowCondition)
 	, GetData(InGetData)
 	{}
 	
-	int32 GDaySequenceDebugLevel = 0;
+	int32 GDaySequenceDebugLevel = 2;
 	FAutoConsoleVariableRef CVarDaySequenceActorDebugLevel(
 		TEXT("DaySequence.DebugLevel"),
 		GDaySequenceDebugLevel,
@@ -123,7 +123,7 @@ ADaySequenceActor::ADaySequenceActor(const FObjectInitializer& Init)
 	bReplicatePlayback = true;
 	SetHidden(false);
 	
-#if ENABLE_DRAW_DEBUG
+#if DAY_SEQUENCE_ENABLE_DRAW_DEBUG
 	if (!IsTemplate())
 	{
 		AHUD::OnShowDebugInfo.AddUObject(this, &ADaySequenceActor::OnShowDebugInfo);
@@ -179,7 +179,7 @@ void ADaySequenceActor::Tick(float DeltaTime)
 	}
 #endif
 
-#if ENABLE_DRAW_DEBUG
+#if DAY_SEQUENCE_ENABLE_DRAW_DEBUG
 	if (UE::DaySequence::GDaySequenceDebugLevel != CachedDebugLevel)
 	{
 		OnDebugLevelChanged.Broadcast(UE::DaySequence::GDaySequenceDebugLevel);
@@ -545,13 +545,8 @@ void ADaySequenceActor::InitializeRootSequence()
 #if ROOT_SEQUENCE_RECONSTRUCTION_ENABLED
 	}
 #endif
-	
-	if (GetNetMode() == NM_DedicatedServer)
-	{
-		return;
-	}
 
-#if ENABLE_DRAW_DEBUG
+#if DAY_SEQUENCE_ENABLE_DRAW_DEBUG
 	for (TSharedPtr<UE::DaySequence::FDaySequenceDebugEntry> Entry : SubSectionDebugEntries)
 	{
 		UnregisterDebugEntry(Entry, ShowDebug_SubSequenceCategory);
@@ -560,6 +555,11 @@ void ADaySequenceActor::InitializeRootSequence()
 	SubSectionDebugEntries.Empty();
 #endif
 	
+	if (GetNetMode() == NM_DedicatedServer)
+	{
+		return;
+	}
+	
 #if ROOT_SEQUENCE_RECONSTRUCTION_ENABLED
 	if (bInEditorWorld)
 	{
@@ -613,7 +613,7 @@ void ADaySequenceActor::InitializeRootSequence()
 	OnPostInitializeDaySequences.Broadcast(nullptr);
 #endif
 
-#if ENABLE_DRAW_DEBUG
+#if DAY_SEQUENCE_ENABLE_DRAW_DEBUG
 	if (!IsDebugCategoryRegistered(ShowDebug_SubSequenceCategory))
 	{
 		RegisterDebugCategory(ShowDebug_SubSequenceCategory,  OnShowDebugInfoDrawFunction);
@@ -793,7 +793,7 @@ UMovieSceneSubSection* ADaySequenceActor::InitializeDaySequence(const FDaySequen
 		UE_LOG(LogDaySequence, Warning, TEXT("Failed to create UDaySequenceTrack in ADaySequenceActor::InitializeDaySequence"));
 	}
 
-#if ENABLE_DRAW_DEBUG
+#if DAY_SEQUENCE_ENABLE_DRAW_DEBUG
 	if (SubSection)
 	{
         TSharedPtr<TMap<FString, FString>> DebugData = MakeShared<TMap<FString, FString>>();
@@ -803,7 +803,10 @@ UMovieSceneSubSection* ADaySequenceActor::InitializeDaySequence(const FDaySequen
         {
         	if (IsValid(SubSection))
         	{
-        		(*DebugData).FindOrAdd("Actor Name") = GetFName().ToString();
+        		(*DebugData).FindOrAdd("Actor") = GetFName().ToString();
+        		(*DebugData).FindOrAdd("Local Role") = StaticEnum<ENetRole>()->GetNameStringByValue(GetLocalRole());
+        		(*DebugData).FindOrAdd("Remote Role") = StaticEnum<ENetRole>()->GetNameStringByValue(GetRemoteRole());
+        		(*DebugData).FindOrAdd("Authority") = HasAuthority() ? "True" : "False";
         		(*DebugData).FindOrAdd("Sequence Name") = SubSection->GetSequence() ? SubSection->GetSequence()->GetFName().ToString() : "None";
         		(*DebugData).FindOrAdd("Mute State") = SubSection->IsActive() ? "Active" : "Muted";
         		(*DebugData).FindOrAdd("Hierarchical Bias") = FString::Printf(TEXT("%d"), SubSection->Parameters.HierarchicalBias);
@@ -1360,7 +1363,7 @@ FMovieSceneSequencePlaybackSettings ADaySequenceActor::GetPlaybackSettings(const
 	return Settings;
 }
 
-#if ENABLE_DRAW_DEBUG
+#if DAY_SEQUENCE_ENABLE_DRAW_DEBUG
 bool ADaySequenceActor::IsDebugCategoryRegistered(const FName& Category) const
 {
 	return DebugEntries.Find(Category) != nullptr;
@@ -1539,6 +1542,14 @@ void ADaySequenceActor::OnShowDebugInfo(AHUD* HUD, UCanvas* Canvas, const FDebug
 		return;
 	}
 
+#if WITH_EDITOR
+	// Necessary because we actually get called from a static delegate, so this can be called for editor & PIE actor which obfuscates the data.
+	if (const UWorld* World = GetWorld(); World && World->WorldType == EWorldType::Editor)
+	{
+		return;
+	}
+#endif
+
 	DebugEntries = DebugEntries.FilterByPredicate([](const TPair<FName, TPair<FDebugEntryArray, FDebugCategoryDrawFunction>>& Category) {
 		return !Category.Value.Key.IsEmpty();
 	});
diff --git a/Engine/Plugins/Experimental/DaySequence/Source/DaySequence/Private/DaySequenceModifierComponent.cpp b/Engine/Plugins/Experimental/DaySequence/Source/DaySequence/Private/DaySequenceModifierComponent.cpp
--- a/Engine/Plugins/Experimental/DaySequence/Source/DaySequence/Private/DaySequenceModifierComponent.cpp
+++ b/Engine/Plugins/Experimental/DaySequence/Source/DaySequence/Private/DaySequenceModifierComponent.cpp
@@ -138,25 +138,8 @@ UDaySequenceModifierComponent::UDaySequenceModifierComponent(const FObjectInitia
 	BlendAmount = 100.f;
 	UserBlendWeight = 1.f;
 	InternalBlendWeight = 1.f;
-#if ENABLE_DRAW_DEBUG
+#if DAY_SEQUENCE_ENABLE_DRAW_DEBUG
 	DebugLevel = 0;
-	
-	// This gets captured by a lambda below so should continue living
-	TSharedPtr<TMap<FString, FString>> DebugData = MakeShared<TMap<FString, FString>>();
-	DebugEntry = MakeShared<UE::DaySequence::FDaySequenceDebugEntry>(
-	[this](){ return ShouldShowDebugInfo(); },
-	[this, DebugData]()
-	{
-		(*DebugData).FindOrAdd("Owner Name") = GetOwner()->GetFName().ToString();
-		(*DebugData).FindOrAdd("Component Enabled") = bIsComponentEnabled ? "True" : "False";
-		(*DebugData).FindOrAdd("Modifier Enabled") = bIsEnabled ? "True" : "False";
-		(*DebugData).FindOrAdd("Blend Weight") = FString::Printf(TEXT("%.5f"), GetBlendWeight());
-
-		const APlayerController* BlendTarget = WeakBlendTarget.Get();
-		(*DebugData).FindOrAdd("Blend Target" ) = BlendTarget ? BlendTarget->GetName() : "None";
-
-		return DebugData;
-	});
 #endif
 	
 	PrimaryComponentTick.bCanEverTick = false;
@@ -347,13 +330,34 @@ void UDaySequenceModifierComponent::BindToDaySequenceActor(ADaySequenceActor* Da
 		
 		DaySequenceActor->GetOnPostInitializeDaySequences().AddUObject(this, &UDaySequenceModifierComponent::ReinitializeSubSequence);
 		DaySequenceActor->GetOnDaySequenceUpdate().AddUObject(this, &UDaySequenceModifierComponent::DaySequenceUpdate);
-#if ENABLE_DRAW_DEBUG
+#if DAY_SEQUENCE_ENABLE_DRAW_DEBUG
 		if (!DaySequenceActor->IsDebugCategoryRegistered(ShowDebug_ModifierCategory))
 		{
 			DaySequenceActor->RegisterDebugCategory(ShowDebug_ModifierCategory, TargetActor->OnShowDebugInfoDrawFunction);
 		}
 		
 		DaySequenceActor->GetOnDebugLevelChanged().AddUObject(this, &UDaySequenceModifierComponent::OnDebugLevelChanged);
+		DebugLevel = DaySequenceActor->GetDebugLevel();
+	
+		// This gets captured by a lambda below so should continue living
+		TSharedPtr<TMap<FString, FString>> DebugData = MakeShared<TMap<FString, FString>>();
+		DebugEntry = MakeShared<UE::DaySequence::FDaySequenceDebugEntry>(
+		[this](){ return ShouldShowDebugInfo(); },
+		[this, DebugData]()
+		{
+			(*DebugData).FindOrAdd("Actor") = GetOwner()->GetFName().ToString();
+			(*DebugData).FindOrAdd("Local Role") = StaticEnum<ENetRole>()->GetNameStringByValue(GetOwner()->GetLocalRole());
+			(*DebugData).FindOrAdd("Remote Role") = StaticEnum<ENetRole>()->GetNameStringByValue(GetOwner()->GetRemoteRole());
+			(*DebugData).FindOrAdd("Component Enabled") = bIsComponentEnabled ? "True" : "False";
+			(*DebugData).FindOrAdd("Modifier Enabled") = bIsEnabled ? "True" : "False";
+			(*DebugData).FindOrAdd("Blend Weight") = FString::Printf(TEXT("%.5f"), GetBlendWeight());
+
+			const APlayerController* BlendTarget = WeakBlendTarget.Get();
+			(*DebugData).FindOrAdd("Blend Target" ) = BlendTarget ? BlendTarget->GetName() : "None";
+
+			return DebugData;
+		});
+		
 		DaySequenceActor->RegisterDebugEntry(DebugEntry, ShowDebug_ModifierCategory);
 #endif
 	}
@@ -373,9 +377,10 @@ void UDaySequenceModifierComponent::UnbindFromDaySequenceActor()
 	{
 		TargetActor->GetOnPostInitializeDaySequences().RemoveAll(this);
 		TargetActor->GetOnDaySequenceUpdate().RemoveAll(this);
-#if ENABLE_DRAW_DEBUG
+#if DAY_SEQUENCE_ENABLE_DRAW_DEBUG
 		TargetActor->GetOnDebugLevelChanged().RemoveAll(this);
 		TargetActor->UnregisterDebugEntry(DebugEntry, ShowDebug_ModifierCategory);
+		DebugEntry.Reset();
 #endif
 		TargetActor = nullptr;
 	}
@@ -406,7 +411,7 @@ void UDaySequenceModifierComponent::RemoveSubSequenceTrack()
 	}
 	SubSections.Empty();
 
-#if ENABLE_DRAW_DEBUG
+#if DAY_SEQUENCE_ENABLE_DRAW_DEBUG
 	if (TargetActor)
 	{
 		for (TSharedPtr<UE::DaySequence::FDaySequenceDebugEntry> Entry : SubSectionDebugEntries)
@@ -665,7 +670,7 @@ void UDaySequenceModifierComponent::ReinitializeSubSequence(ADaySequenceActor::F
 	}
 #endif
 	
-#if ENABLE_DRAW_DEBUG
+#if DAY_SEQUENCE_ENABLE_DRAW_DEBUG
 	if (TargetActor)
 	{
 		if (!TargetActor->IsDebugCategoryRegistered(TargetActor->ShowDebug_SubSequenceCategory))
@@ -800,16 +805,19 @@ UMovieSceneSubSection* UDaySequenceModifierComponent::InitializeDaySequence(cons
 	OnInvalidateMuteStates.AddWeakLambda(SubSection, SetSubTrackMuteState);
 	TargetActor->BindToConditionCallbacks(this, Entry.Conditions.Conditions, [this]() { InvalidateMuteStates(); });
 	
-#if ENABLE_DRAW_DEBUG
+#if DAY_SEQUENCE_ENABLE_DRAW_DEBUG	
 	// This gets captured by a lambda below so should continue living
 	TSharedPtr<TMap<FString, FString>> DebugData = MakeShared<TMap<FString, FString>>();
 	SubSectionDebugEntries.Emplace(MakeShared<UE::DaySequence::FDaySequenceDebugEntry>(
-	[this](){ return true; },
+	[this](){ return ShouldShowDebugInfo(); },
 	[this, DebugData, SubSection]()
 	{
 		if (IsValid(SubSection))
 		{
-			(*DebugData).FindOrAdd("Owner Name") = GetOwner()->GetFName().ToString();
+			(*DebugData).FindOrAdd("Actor") = GetOwner()->GetFName().ToString();
+			(*DebugData).FindOrAdd("Local Role") = StaticEnum<ENetRole>()->GetNameStringByValue(GetOwner()->GetLocalRole());
+			(*DebugData).FindOrAdd("Remote Role") = StaticEnum<ENetRole>()->GetNameStringByValue(GetOwner()->GetRemoteRole());
+			(*DebugData).FindOrAdd("Authority") = GetOwner()->HasAuthority() ? "True" : "False";
 			(*DebugData).FindOrAdd("Sequence Name") = SubSection->GetSequence() ? SubSection->GetSequence()->GetFName().ToString() : "None";
 			(*DebugData).FindOrAdd("Mute State") = SubSection->IsActive() ? "Active" : "Muted";
 			(*DebugData).FindOrAdd("Hierarchical Bias") = FString::Printf(TEXT("%d"), SubSection->Parameters.HierarchicalBias);
@@ -1075,7 +1083,7 @@ void UDaySequenceModifierComponent::InvalidateMuteStates() const
 	OnInvalidateMuteStates.Broadcast();
 }
 
-#if ENABLE_DRAW_DEBUG
+#if DAY_SEQUENCE_ENABLE_DRAW_DEBUG
 void UDaySequenceModifierComponent::OnDebugLevelChanged(int32 InDebugLevel)
 {
 	DebugLevel = InDebugLevel;
@@ -1083,6 +1091,11 @@ void UDaySequenceModifierComponent::OnDebugLevelChanged(int32 InDebugLevel)
 
 bool UDaySequenceModifierComponent::ShouldShowDebugInfo() const
 {
+	if (!TargetActor || TargetActor->GetNetMode() == NM_DedicatedServer)
+	{
+		return false;
+	}
+	
 	switch (DebugLevel)
 	{
 	case 0: return false;
diff --git a/Engine/Plugins/Experimental/DaySequence/Source/DaySequence/Public/DaySequenceActor.h b/Engine/Plugins/Experimental/DaySequence/Source/DaySequence/Public/DaySequenceActor.h
--- a/Engine/Plugins/Experimental/DaySequence/Source/DaySequence/Public/DaySequenceActor.h
+++ b/Engine/Plugins/Experimental/DaySequence/Source/DaySequence/Public/DaySequenceActor.h
@@ -4,7 +4,6 @@
 
 #include "GameFramework/Info.h"
 #include "DaySequenceConditionSet.h"
-#include "DrawDebugHelpers.h"	// Defines ENABLE_DRAW_DEBUG
 #include "IDaySequencePlayer.h"
 #include "IMovieScenePlaybackClient.h"
 #include "MovieSceneBindingOwnerInterface.h"
@@ -14,12 +13,20 @@
 
 #include "DaySequenceActor.generated.h"
 
-namespace EEndPlayReason { enum Type : int; }
+#ifndef ROOT_SEQUENCE_RECONSTRUCTION_ENABLED
+	#define ROOT_SEQUENCE_RECONSTRUCTION_ENABLED WITH_EDITOR
+#endif
 
-#if ENABLE_DRAW_DEBUG
-class AHUD;
+#ifndef DAY_SEQUENCE_ENABLE_DRAW_DEBUG
+	#define DAY_SEQUENCE_ENABLE_DRAW_DEBUG !UE_BUILD_SHIPPING
 #endif
 
+#if DAY_SEQUENCE_ENABLE_DRAW_DEBUG
+	class AHUD;
+#endif
+
+namespace EEndPlayReason { enum Type : int; }
+
 class UCurveFloat;
 class UDaySequence;
 class UDaySequenceCollectionAsset;
@@ -33,16 +40,12 @@ class FDebugDisplayInfo;
 struct FDaySequenceCollectionEntry;
 struct FMovieSceneSequencePlaybackSettings;
 
-#ifndef ROOT_SEQUENCE_RECONSTRUCTION_ENABLED
-	#define ROOT_SEQUENCE_RECONSTRUCTION_ENABLED WITH_EDITOR
-#endif
-
 namespace UE::DaySequence
 {
 	struct FStaticTimeContributor;
 	struct FStaticTimeManager;
 
-#if ENABLE_DRAW_DEBUG
+#if DAY_SEQUENCE_ENABLE_DRAW_DEBUG
 	// This provides methods for determining if this debug entry should be shown and for getting a pointer to the debug data.
 	// Anything (Currently only DaySequenceModifierComponents) can submit one of these entries to a DaySequenceActor.
 	// The debug data can be printed in play with the command "showdebug DaySequence" if ShowCondition evaluates to true. 
@@ -287,13 +290,14 @@ public:
 	
 	void InvalidateMuteStates() const;
 	
-#if ENABLE_DRAW_DEBUG
+#if DAY_SEQUENCE_ENABLE_DRAW_DEBUG
 	const FName ShowDebug_GeneralCategory = "DaySequence";
 	const FName ShowDebug_SubSequenceCategory = "DaySequenceSubSequences";
 	
 	DECLARE_EVENT_OneParam(ADaySequenceActor, FOnDebugLevelChanged, int32)
 	FOnDebugLevelChanged& GetOnDebugLevelChanged() { return OnDebugLevelChanged; }
-
+	int32 GetDebugLevel() const { return CachedDebugLevel; }
+	
 	bool IsDebugCategoryRegistered(const FName& Category) const;
 	void RegisterDebugCategory(const FName& Category, UE::DaySequence::FDebugCategoryDrawFunction DrawFunction);
 
@@ -536,7 +540,7 @@ protected:
 	UFUNCTION()
 	void StopDaySequenceUpdateTimer();
 	
-#if ENABLE_DRAW_DEBUG
+#if DAY_SEQUENCE_ENABLE_DRAW_DEBUG
 	void OnShowDebugInfo(AHUD* HUD, UCanvas* Canvas, const FDebugDisplayInfo& DisplayInfo, float& YL, float& YPos);
 
 	// The keys in this map are debug categories. When a debug entry is registered, it is is added
diff --git a/Engine/Plugins/Experimental/DaySequence/Source/DaySequence/Public/DaySequenceModifierComponent.h b/Engine/Plugins/Experimental/DaySequence/Source/DaySequence/Public/DaySequenceModifierComponent.h
--- a/Engine/Plugins/Experimental/DaySequence/Source/DaySequence/Public/DaySequenceModifierComponent.h
+++ b/Engine/Plugins/Experimental/DaySequence/Source/DaySequence/Public/DaySequenceModifierComponent.h
@@ -8,7 +8,6 @@
 #include "Containers/Ticker.h"
 #include "DaySequenceActor.h"
 #include "DaySequenceConditionSet.h"
-#include "DrawDebugHelpers.h"	// Defines ENABLE_DRAW_DEBUG
 #include "Generators/MovieSceneEasingFunction.h"
 
 #include "DaySequenceModifierComponent.generated.h"
@@ -25,7 +24,7 @@ class UMovieSceneSubSection;
 
 struct FDaySequenceCollectionEntry;
 
-#if ENABLE_DRAW_DEBUG
+#if DAY_SEQUENCE_ENABLE_DRAW_DEBUG
 namespace UE::DaySequence
 {
 	struct FDaySequenceDebugEntry;
@@ -374,8 +373,8 @@ private:
 
 	UE::DaySequence::FOnInvalidateMuteStates OnInvalidateMuteStates;
 	
-#if ENABLE_DRAW_DEBUG
-	const FName ShowDebug_ModifierCategory = "DaySequence_Modifiers";
+#if DAY_SEQUENCE_ENABLE_DRAW_DEBUG
+	const FName ShowDebug_ModifierCategory = "DaySequenceModifiers";
 	
 	void OnDebugLevelChanged(int32 InDebugLevel);
 	bool ShouldShowDebugInfo() const;
