DaySequence: Fixed DaySequence actor preview to only update for Editor world cleanup.

diff --git a/Engine/Plugins/Experimental/DaySequence/Source/DaySequenceEditor/Private/DaySequenceActorPreview.cpp b/Engine/Plugins/Experimental/DaySequence/Source/DaySequenceEditor/Private/DaySequenceActorPreview.cpp
--- a/Engine/Plugins/Experimental/DaySequence/Source/DaySequenceEditor/Private/DaySequenceActorPreview.cpp
+++ b/Engine/Plugins/Experimental/DaySequence/Source/DaySequenceEditor/Private/DaySequenceActorPreview.cpp
@@ -322,7 +322,10 @@ void FDaySequenceActorPreview::OnBeginPIE(bool /*bSimulate*/)
 
 void FDaySequenceActorPreview::OnWorldCleanup(UWorld* World, bool bSessionEnded, bool bCleanupResources)
 {
-	EnablePreview(false);
+	if (World && World->WorldType == EWorldType::Editor)
+	{
+		EnablePreview(false);
+	}
 }
 
 void FDaySequenceActorPreview::UpdateLevelEditorViewports()
