SunMoonDaySequenceActor: Fix defaults

diff --git a/Engine/Plugins/Experimental/DaySequence/Source/DaySequence/Private/Actors/SunMoonDaySequenceActor.cpp b/Engine/Plugins/Experimental/DaySequence/Source/DaySequence/Private/Actors/SunMoonDaySequenceActor.cpp
--- a/Engine/Plugins/Experimental/DaySequence/Source/DaySequence/Private/Actors/SunMoonDaySequenceActor.cpp
+++ b/Engine/Plugins/Experimental/DaySequence/Source/DaySequence/Private/Actors/SunMoonDaySequenceActor.cpp
@@ -21,10 +21,10 @@ ASunMoonDaySequenceActor::ASunMoonDaySequenceActor(const FObjectInitializer& Ini
 
 	// Configure other Moon defaults.
 	MoonComponent->SetAtmosphereSunLightIndex(1);	// Make Moon the secondary directional light that contributes to the sky atmosphere.
-	MoonComponent->SetIntensity(2.f);
+	MoonComponent->SetIntensity(0.05f);
 	MoonComponent->SetUseTemperature(true);
 	MoonComponent->SetTemperature(9000.f);
-	MoonComponent->SetWorldRotation(FRotator(0.f, 45.f, 0.f));
+	MoonComponent->SetWorldRotation(FRotator(-45.f, 0.f, 0.f));
 	
 	// Override the sky sphere material.
 	static ConstructorHelpers::FObjectFinder<UMaterialInterface> SkySphereDefaultMaterial(TEXT("/DaySequence/MI_24hrSky.MI_24hrSky"));
