SDaySequencePreviewTimeSlider: Instead of directly setting preview time on the actor preview, set it on the underlying DSA.

This causes TimeOfDayPreview to actually get set. The data still gets to the actor preview via the delegates broadcast in SetTimeOfDayPreview

diff --git a/Engine/Plugins/Experimental/DaySequence/Source/DaySequenceEditor/Private/SDaySequencePreviewTimeSlider.cpp b/Engine/Plugins/Experimental/DaySequence/Source/DaySequenceEditor/Private/SDaySequencePreviewTimeSlider.cpp
--- a/Engine/Plugins/Experimental/DaySequence/Source/DaySequenceEditor/Private/SDaySequencePreviewTimeSlider.cpp
+++ b/Engine/Plugins/Experimental/DaySequence/Source/DaySequenceEditor/Private/SDaySequencePreviewTimeSlider.cpp
@@ -25,14 +25,7 @@ void SDaySequencePreviewTimeSlider::Construct(const FArguments& InArgs)
 			.IsEnabled_Lambda([this]()
 			{
 				const bool bPreviewEnabled = DaySequenceActorPreview->IsPreviewEnabled();
-				bool bHasStaticTime = false;
-				
-				if (const ADaySequenceActor* PreviewActor = DaySequenceActorPreview->GetPreviewActor().Get())
-				{
-					bHasStaticTime = PreviewActor->HasStaticTimeOfDay();
-				}
-				
-				return bPreviewEnabled && !bHasStaticTime;
+				return bPreviewEnabled;
 			})
 			.AllowPan(false)
 			.AllowZoom(false)
@@ -60,9 +53,9 @@ void SDaySequencePreviewTimeSlider::Construct(const FArguments& InArgs)
 			.OnScrubPositionChanged_Lambda(
 				[this](double NewScrubTime, bool /*bIsScrubbing*/)
 				{
-					if (DaySequenceActorPreview->GetPreviewSequencer().IsValid())
+					if (ADaySequenceActor* PreviewActor = DaySequenceActorPreview->GetPreviewActor().Get())
 					{
-						DaySequenceActorPreview->SetPreviewTime(NewScrubTime);
+						PreviewActor->SetTimeOfDayPreview(NewScrubTime);
 					}
 				})
 		]
