SDaySequenceTimeSlider: Undoing a portion of the previous commit that wasn't intended to be included.

We will likely do something like this but this was submitted without review or rigorous testing by accident.

diff --git a/Engine/Plugins/Experimental/DaySequence/Source/DaySequenceEditor/Private/SDaySequencePreviewTimeSlider.cpp b/Engine/Plugins/Experimental/DaySequence/Source/DaySequenceEditor/Private/SDaySequencePreviewTimeSlider.cpp
--- a/Engine/Plugins/Experimental/DaySequence/Source/DaySequenceEditor/Private/SDaySequencePreviewTimeSlider.cpp
+++ b/Engine/Plugins/Experimental/DaySequence/Source/DaySequenceEditor/Private/SDaySequencePreviewTimeSlider.cpp
@@ -25,7 +25,14 @@ void SDaySequencePreviewTimeSlider::Construct(const FArguments& InArgs)
 			.IsEnabled_Lambda([this]()
 			{
 				const bool bPreviewEnabled = DaySequenceActorPreview->IsPreviewEnabled();
-				return bPreviewEnabled;
+				bool bHasStaticTime = false;
+				
+				if (const ADaySequenceActor* PreviewActor = DaySequenceActorPreview->GetPreviewActor().Get())
+				{
+					bHasStaticTime = PreviewActor->HasStaticTimeOfDay();
+				}
+				
+				return bPreviewEnabled && !bHasStaticTime;
 			})
 			.AllowPan(false)
 			.AllowZoom(false)
