DayLengthStaticTime: Make static time blending take the fastest path from current global time to desired static time.

Lonnie wrote the latest revision of this.

diff --git a/Engine/Plugins/Experimental/DaySequence/Source/DaySequence/Private/DaySequenceActor.cpp b/Engine/Plugins/Experimental/DaySequence/Source/DaySequence/Private/DaySequenceActor.cpp
--- a/Engine/Plugins/Experimental/DaySequence/Source/DaySequence/Private/DaySequenceActor.cpp
+++ b/Engine/Plugins/Experimental/DaySequence/Source/DaySequence/Private/DaySequenceActor.cpp
@@ -1216,7 +1216,7 @@ float ADaySequenceActor::GetStaticTimeOfDay() const
 {
 	if (StaticTimeManager && HasStaticTimeOfDay())
 	{
-		return StaticTimeManager->GetStaticTime(GetTimeOfDay());
+		return StaticTimeManager->GetStaticTime(GetTimeOfDay(), DayLength.ToHours());
 	}
 	
 	return std::numeric_limits<float>::lowest();
diff --git a/Engine/Plugins/Experimental/DaySequence/Source/DaySequence/Private/DaySequenceStaticTime.cpp b/Engine/Plugins/Experimental/DaySequence/Source/DaySequence/Private/DaySequenceStaticTime.cpp
--- a/Engine/Plugins/Experimental/DaySequence/Source/DaySequence/Private/DaySequenceStaticTime.cpp
+++ b/Engine/Plugins/Experimental/DaySequence/Source/DaySequence/Private/DaySequenceStaticTime.cpp
@@ -2,6 +2,8 @@
 
 #include "DaySequenceStaticTime.h"
 
+#include "DaySequenceModule.h"
+
 namespace UE::DaySequence
 {
 	void FStaticTimeManager::AddStaticTimeContributor(const FStaticTimeContributor& NewContributor)
@@ -62,14 +64,20 @@ namespace UE::DaySequence
 			}
 		}
 
+		// Reset the blend data when static time goes away.
+		LastBlendDirection.Reset();
+		LastBlendDelta.Reset();
+		LastBlendWinding = 0;
+		LastBlendOffset = 0;
+
 		return false;
 	}
 	
-	float FStaticTimeManager::GetStaticTime(float InitialTime) const
+	float FStaticTimeManager::GetStaticTime(float InitialTime, float DayLength) const
 	{
 		float AccumulatedWeight = 0.f;
 		float AccumulatedTime = 0.f;
-			
+		
 		// Process batches of contributors based on priority
 		for (int32 PriorityGroupStartIdx = 0; PriorityGroupStartIdx < Contributors.Num();)
 		{
@@ -91,18 +99,153 @@ namespace UE::DaySequence
 			PriorityGroupStartIdx = PriorityGroupEndIdx + 1;
 		}
 
-		// Blend against initial value if necessary
-		if (AccumulatedWeight < 1.f)
+		// The implementation below aims to provide the shortest Lerp path from the InitialTime to
+		// the TargetTime based on the AccumulatedWeight.
+		//
+		// For a static global time (InitialTime), the formula is straightforward. We compute a
+		// InitialTime relative to TargetTime to identify the BlendDelta. BlendDelta defines the range
+		// of our Lerp and will exist in the range [-0.5D, 0.5D], where D = DayLength. Then we can
+		// simply Lerp(InitialTime + BlendDelta, TargetTime, AccumulatedWeight).
+		//
+		// For a dynamic global time (ie. advancing global time), the solution is more involved. With
+		// dynamic InitialTime, the Lerp range is continuously changing and may cross the [-0.5D, 0.5D]
+		// boundary. This boundary is discontinuous and can result in unsightly pops in the blended
+		// static time when AccumulatedWeight < 1.f.
+		//
+		// One method of solving that is to continuously add +/-1D in the direction that time is travelling
+		// when crossing the [-0.5D, 0.5D] boundary. This would effectively result in a Lerp range of
+		// [TargetTime, InitialTime +/- N*DayCycle] depending on the accumulated passage of time. One
+		// consequence of that approach however is that since the Lerp now spans multiple DayCycles,
+		// walking between ends of the blend region will cause up to N DayCycles to wind/unwind.
+		//
+		// The solution used here accounts for the above problem by leveraging the circular nature
+		// of day cycles. Specifically, we reframe the problem space as [-1D, 1D]
+		//
+		//    |--------|--------|--------|--------|
+		//   -1D     -0.5D     Tgt      0.5D      1D
+		//
+		// We want to solve for a continuous Lerp through the above [-1D, 1D] space:
+		//
+		//    Lerp(StartTime, EndTime, AccumulatedWeight)
+		//    StartTime = InitialTime + BlendDelta
+		//
+		// If we focus on the forward advancement of time case:
+		//
+		// 1. When we cross the +0.5D mark, BlendDelta will flip from +12 to -12. To maintain a continuity
+		//    we add a BlendOffset of 1D:
+		//
+		//    StartTime = InitialTime + BlendDelta + BlendOffset * DayCycle
+		//    BlendOffset = 1
+		//
+		//    This allows StartTime to continuously increment from +12 --> +13 and all the way up to +24
+		//    at the 1D mark.
+		//
+		// 2. When we cross the 1D mark, we mirror the problem to the negative space since 1D == -1D.
+		//    However, the trick is that since the negative side of TargetTime is mirrored compared
+		//    to the positive side, we must flip the lerp so that the lerp remains continuous with
+		//    our input AccumulatedWeight.
+		//
+		// 3. Lastly, when we cross past TargetTime again, we need to stop the flipped lerp since
+		//    we'll be on the positive side of TargetTime once again.
+		//    
+		// The same applies to the backwards advancement of time case just in reverse. So instead of adding
+		// a BlendOffset of 1D, we subtract.
+		//
+		const float HalfDayLength = 0.5f * DayLength;
+		float TargetTime = AccumulatedWeight > 0.f ? AccumulatedTime / AccumulatedWeight : InitialTime;
+		
+		// Returns the shortest delta to TargetTime in the range [TargetTime - HalfDayLength, TargetTime + HalfDayLength]
+		auto GetBlendDelta = [TargetTime, InitialTime, DayLength, HalfDayLength]()
+		{
+			float TimeDelta = InitialTime - TargetTime;
+			if (TimeDelta > HalfDayLength)
+			{
+				TimeDelta -= DayLength;
+			}
+			else if(TimeDelta < -HalfDayLength)
+			{
+				TimeDelta += DayLength;
+			}
+			return TimeDelta;
+		};
+
+		// Returns 1 or -1 if BlendDelta is on the positive or negative side of TargetTime, respectively.
+		auto GetBlendDirection = [](float Delta)
 		{
-			const float FillWeight = (1.f - AccumulatedWeight);
-			const float FillTime = InitialTime * FillWeight;
+			return Delta >= 0.f ? 1 : -1;
+		};
 
-			AccumulatedWeight += FillWeight;
-			AccumulatedTime += FillTime;
+		const float BlendDelta = GetBlendDelta();
+		const int BlendDirection = GetBlendDirection(BlendDelta);
+
+		// Compute LastBlendOffset based on the number of times (incl. direction) BlendDelta
+		// has crossed the TargetTime +/- HalfDayLength mark.
+		if (LastBlendDelta.IsSet() &&
+			LastBlendDirection.IsSet() &&
+			LastBlendDirection.GetValue() != BlendDirection)
+		{
+			const float BlendDeltaChange = BlendDelta - LastBlendDelta.GetValue();
+			if (FMath::Abs(BlendDeltaChange) > HalfDayLength)
+			{
+				// Crossed TargetTime +- HalfDayLength
+				
+				// When crossing the TargetTime +- HalfDayLength, the sign is inverted from the BlendDeltaChange
+				// due to the wraparound effect.
+				const int BlendDeltaSign = BlendDeltaChange < 0.f ? 1 : -1;
+				
+				// Record our winding counts based on direction of the crossover.
+				LastBlendWinding += BlendDeltaSign;
+				LastBlendOffset = FMath::WrapExclusive(LastBlendOffset + BlendDeltaSign, -1, 2);
+			}
+			else
+			{
+				// Crossed TargetTime
+				if (LastBlendOffset != 0)
+				{
+					// Loop TargetTime + HalfDayLength  <--> TargetTime - HalfDayLength
+					const int BlendDeltaSign = BlendDeltaChange > 0.f ? 1 : -1;
+					LastBlendOffset = FMath::WrapExclusive(LastBlendOffset + BlendDeltaSign, -1, 2);
+				}
+			}
 		}
 
+		// Blend against initial value if necessary
+		if (AccumulatedWeight < 1.0f)
+		{
+			float BlendOffset = LastBlendOffset * DayLength;
+			const float StartTime = TargetTime + BlendDelta + BlendOffset;
+
+			// If we have applied a blend offset in the opposite direction of our LastBlendWinding and remain
+			// in the opposite space of the BlendWinding (ie. BlendDelta < 0.f for +Winding or BlendDelta > 0.f for -Winding)
+			// then we must invert the Lerp since the weight is mirrored across TargetTime.
+			if ((LastBlendWinding > 0 && (LastBlendOffset < 0 || (LastBlendOffset == 0 && BlendDelta < 0.f))) ||
+				(LastBlendWinding < 0 && (LastBlendOffset > 0 || (LastBlendOffset == 0 && BlendDelta > 0.f))))
+			{
+				AccumulatedTime = FMath::Lerp(TargetTime, StartTime, AccumulatedWeight);
+			}
+			else
+			{
+				AccumulatedTime = FMath::Lerp(StartTime, TargetTime, AccumulatedWeight);
+			}
+
+			// Wrap is inclusive of the upper bound. Ensure we wrap around correctly.
+			AccumulatedTime = FMath::Wrap(AccumulatedTime, 0.f, DayLength);
+			if (AccumulatedTime >= DayLength)
+			{
+				AccumulatedTime -= DayLength;
+			}
+
+#if !NO_LOGGING
+			UE_LOG(LogDaySequence, VeryVerbose, TEXT("BlendDelta: %f | BlendDir: %d | BlendWinding: %d | BlendOffset: %f | StartTime: %f | InitialTime: %f | TargetTime: %f"), BlendDelta, BlendDirection, LastBlendWinding, BlendOffset, StartTime, InitialTime, TargetTime);
+			UE_LOG(LogDaySequence, VeryVerbose, TEXT("AccumTime: %f | AccumWeight: %f"), AccumulatedTime, AccumulatedWeight);
+#endif
+		}
+		
 		// we could probably ensure(AccumulatedWeight is nearly equal to 1.f), leaving out for now
 
+		LastBlendDelta = BlendDelta;
+		LastBlendDirection = BlendDirection;
+		
 		return AccumulatedTime;
 	}
 
diff --git a/Engine/Plugins/Experimental/DaySequence/Source/DaySequence/Private/DaySequenceStaticTime.h b/Engine/Plugins/Experimental/DaySequence/Source/DaySequence/Private/DaySequenceStaticTime.h
--- a/Engine/Plugins/Experimental/DaySequence/Source/DaySequence/Private/DaySequenceStaticTime.h
+++ b/Engine/Plugins/Experimental/DaySequence/Source/DaySequence/Private/DaySequenceStaticTime.h
@@ -46,10 +46,10 @@ namespace UE::DaySequence
 		void AddStaticTimeContributor(const FStaticTimeContributor& NewContributor);
 		void RemoveStaticTimeContributor(const UObject* UserObject);
 		bool HasStaticTime() const;
-		float GetStaticTime(float InitialTime) const;
+		float GetStaticTime(float InitialTime, float DayLength) const;
 		
 	private:
-
+		
 		// Handles multiple contributors of the same priority by averaging their weights and times and returning a single desired weight and time
 		FStaticTimeInfo ProcessPriorityGroup(int32 StartIdx, int32 EndIdx) const;
 		
@@ -57,6 +57,12 @@ namespace UE::DaySequence
 
 		// We precompute this when mutating the Contributors array to avoid redundant reiteration over Contributors when computing static time
 		TMap<int32, int32> PriorityGroupSizes;
+
+		// Cached blend data. Mutable because it is set in (Has/Get)StaticTime.
+		mutable int LastBlendWinding = 0;
+		mutable int LastBlendOffset = 0;
+		mutable TOptional<float> LastBlendDelta;
+		mutable TOptional<int> LastBlendDirection;
 	};
 }
 
