Day Sequencer: AddDynamicSection to DaySequenceEditorDynamic rather than LevelSequenceEditorDynamic.

While investigating the crash for UE-227132, I noticed that the playback context button in the Sequencer UI disappears after existing PIE

#rnx

diff --git a/Engine/Plugins/Experimental/DaySequence/Source/DaySequenceEditor/Private/DaySequenceEditorToolkit.cpp b/Engine/Plugins/Experimental/DaySequence/Source/DaySequenceEditor/Private/DaySequenceEditorToolkit.cpp
--- a/Engine/Plugins/Experimental/DaySequence/Source/DaySequenceEditor/Private/DaySequenceEditorToolkit.cpp
+++ b/Engine/Plugins/Experimental/DaySequence/Source/DaySequenceEditor/Private/DaySequenceEditorToolkit.cpp
@@ -415,7 +415,7 @@ void FDaySequenceEditorToolkit::ExtendSequencerToolbar(FName InToolMenuName)
 	const FToolMenuInsert SectionInsertLocation("BaseCommands", EToolMenuInsertType::Before);
 
 	{
-		ToolMenu->AddDynamicSection("LevelSequenceEditorDynamic", FNewToolMenuDelegate::CreateLambda([](UToolMenu* InMenu)
+		ToolMenu->AddDynamicSection("DaySequenceEditorDynamic", FNewToolMenuDelegate::CreateLambda([](UToolMenu* InMenu)
 		{	
 			UDaySequenceEditorMenuContext* DaySequenceEditorMenuContext = InMenu->FindContext<UDaySequenceEditorMenuContext>();
 			if (DaySequenceEditorMenuContext && DaySequenceEditorMenuContext->Toolkit.IsValid())
