DaySequence: Allow DaySequenceTimeSlider to scrub when StaticTime is active.

diff --git a/Engine/Plugins/Experimental/DaySequence/Source/DaySequenceEditor/Private/SDaySequencePreviewTimeSlider.cpp b/Engine/Plugins/Experimental/DaySequence/Source/DaySequenceEditor/Private/SDaySequencePreviewTimeSlider.cpp
--- a/Engine/Plugins/Experimental/DaySequence/Source/DaySequenceEditor/Private/SDaySequencePreviewTimeSlider.cpp
+++ b/Engine/Plugins/Experimental/DaySequence/Source/DaySequenceEditor/Private/SDaySequencePreviewTimeSlider.cpp
@@ -24,15 +24,7 @@ void SDaySequencePreviewTimeSlider::Construct(const FArguments& InArgs)
 			SNew(SSimpleTimeSlider)
 			.IsEnabled_Lambda([this]()
 			{
-				const bool bPreviewEnabled = DaySequenceActorPreview->IsPreviewEnabled();
-				bool bHasStaticTime = false;
-				
-				if (const ADaySequenceActor* PreviewActor = DaySequenceActorPreview->GetPreviewActor().Get())
-				{
-					bHasStaticTime = PreviewActor->HasStaticTimeOfDay();
-				}
-				
-				return bPreviewEnabled && !bHasStaticTime;
+				return DaySequenceActorPreview->IsPreviewEnabled();
 			})
 			.AllowPan(false)
 			.AllowZoom(false)
