DaySequence: Fixed OnGlobalTimeChanged not propagating to TimeOfDayPreview for actor preview toolkits.

diff --git a/Engine/Plugins/Experimental/DaySequence/Source/DaySequenceEditor/Private/DaySequenceEditorToolkit.cpp b/Engine/Plugins/Experimental/DaySequence/Source/DaySequenceEditor/Private/DaySequenceEditorToolkit.cpp
--- a/Engine/Plugins/Experimental/DaySequence/Source/DaySequenceEditor/Private/DaySequenceEditorToolkit.cpp
+++ b/Engine/Plugins/Experimental/DaySequence/Source/DaySequenceEditor/Private/DaySequenceEditorToolkit.cpp
@@ -304,10 +304,7 @@ void FDaySequenceEditorToolkit::InitializeInternal(
 	
 	const bool bIsSequenceEditor = !IsActorPreview();
 
-	if (bIsSequenceEditor)
-	{
-		Sequencer->OnGlobalTimeChanged().AddSP(this, &FDaySequenceEditorToolkit::OnGlobalTimeChanged);
-	}
+	Sequencer->OnGlobalTimeChanged().AddSP(this, &FDaySequenceEditorToolkit::OnGlobalTimeChanged);
 	
 	FLevelEditorSequencerIntegrationOptions Options;
 	Options.bRequiresLevelEvents = true;
@@ -775,12 +772,13 @@ bool FDaySequenceEditorToolkit::CanFindInContentBrowser() const
 void FDaySequenceEditorToolkit::OnGlobalTimeChanged()
 {
 	// If the sequence we are editing has a DSA outer (i.e. we are editing the root sequence), propagate sequencer time to actor.
-	if (RootActor)
+	// If the sequence we are viewing is the PreviewActor, also propagate sequencer time to actor.
+	if (ADaySequenceActor* DayActor = RootActor ? RootActor : PreviewActor)
 	{
 		// Given sequencer time (SequencerTimeHours), convert to equivalent game time (GameTimeHours).
 		const float SequencerTimeHours = Sequencer->GetGlobalTime().AsSeconds() / 3600.f;
-		const float GameTimeHours = SequencerTimeHours / RootActor->GetTimePerCycle() * RootActor->GetDayLength();
-		RootActor->ConditionalSetTimeOfDayPreview(GameTimeHours);
+		const float GameTimeHours = SequencerTimeHours / DayActor->GetTimePerCycle() * DayActor->GetDayLength();
+		DayActor->ConditionalSetTimeOfDayPreview(GameTimeHours);
 	}
 }
 
