DaySequenceActor: nullcheck BindingOverrides

We are crashing when using duped DSA BPs, issue described here: UE-180613

diff --git a/Engine/Plugins/Experimental/DaySequence/Source/DaySequence/Private/DaySequenceActor.cpp b/Engine/Plugins/Experimental/DaySequence/Source/DaySequence/Private/DaySequenceActor.cpp
--- a/Engine/Plugins/Experimental/DaySequence/Source/DaySequence/Private/DaySequenceActor.cpp
+++ b/Engine/Plugins/Experimental/DaySequence/Source/DaySequence/Private/DaySequenceActor.cpp
@@ -361,7 +361,13 @@ void ADaySequenceActor::Destroyed()
 
 bool ADaySequenceActor::RetrieveBindingOverrides(const FGuid& InBindingId, FMovieSceneSequenceID InSequenceID, TArray<UObject*, TInlineAllocator<1>>& OutObjects) const
 {
-	return BindingOverrides->LocateBoundObjects(InBindingId, InSequenceID, OutObjects);
+	// todo: make this an ensureMsgf? This is null when duping a DSA BP.
+	if (BindingOverrides)
+	{
+		return BindingOverrides->LocateBoundObjects(InBindingId, InSequenceID, OutObjects);
+	}
+
+	return true;
 }
 
 UObject* ADaySequenceActor::GetInstanceData() const
