DaySequenceStaticTime: Fix possible discontinuty in apparent time when entering/exiting the blending region of a day sequence volume

This comes at the cost of a possible discontinuty as global time changes while the player is in the volume, specifically if the player waits in the blending region for a full DayLength. For now we are accepting this as a shortcoming of fixed time blending and advising users to only rely on fixed time blending if they can be reasonably sure this won't happen in practice.

diff --git a/Engine/Plugins/Experimental/DaySequence/Source/DaySequence/Private/DaySequenceStaticTime.cpp b/Engine/Plugins/Experimental/DaySequence/Source/DaySequence/Private/DaySequenceStaticTime.cpp
--- a/Engine/Plugins/Experimental/DaySequence/Source/DaySequence/Private/DaySequenceStaticTime.cpp
+++ b/Engine/Plugins/Experimental/DaySequence/Source/DaySequence/Private/DaySequenceStaticTime.cpp
@@ -65,10 +65,7 @@ namespace UE::DaySequence
 		}
 
 		// Reset the blend data when static time goes away.
-		LastBlendDirection.Reset();
-		LastBlendDelta.Reset();
-		LastBlendWinding = 0;
-		LastBlendOffset = 0;
+		ResetBlendState();
 
 		return false;
 	}
@@ -119,7 +116,10 @@ namespace UE::DaySequence
 		// walking between ends of the blend region will cause up to N DayCycles to wind/unwind.
 		//
 		// The solution used here accounts for the above problem by leveraging the circular nature
-		// of day cycles. Specifically, we reframe the problem space as [-1D, 1D]
+		// of day cycles, at the cost of a discontinuity after 1D of global time has passed while
+		// inside the blended region.
+		//
+		// Specifically, we re-frame the problem space as [-1D, 1D]
 		//
 		//    |--------|--------|--------|--------|
 		//   -1D     -0.5D     Tgt      0.5D      1D
@@ -141,12 +141,21 @@ namespace UE::DaySequence
 		//    at the 1D mark.
 		//
 		// 2. When we cross the 1D mark, we mirror the problem to the negative space since 1D == -1D.
-		//    However, the trick is that since the negative side of TargetTime is mirrored compared
-		//    to the positive side, we must flip the lerp so that the lerp remains continuous with
-		//    our input AccumulatedWeight.
+		//    We do this by incrementing our BlendOffset with a Wrapped range of [-1,1].
+		//
+		// 3. It should be noted that this mirrored problem to the negative space can result in a visible
+		//    discontinuity of up to 0.5D since the mirrored space inverts the shortest path to our target time.
 		//
-		// 3. Lastly, when we cross past TargetTime again, we need to stop the flipped lerp since
-		//    we'll be on the positive side of TargetTime once again.
+		// NOTE: We have tried a number of approaches to resolve the above discontinuity including inverting the lerp
+		//       direction. While inverting the lerp direction ensures a continuous blend as global time passes, it
+		//       results in a discontinuity when leaving the blended region for that entire day cycle. Since the
+		//       lerp is inverted, W=1 would approach InitialTime rather than TargetTime and vice versa creating
+		//       a discontinuity equal to (TargetTime - InitialTime).
+		//
+		//       Ultimately, We made the conscious decision that it is better to ensure a smooth transition into
+		//       and out of the blended region over a continuous blend inside the blended region as global time
+		//       passes. Since our blend state data resets on full entry into the volume or exit out of the volume
+		//       the discontinuity will only show if a player remains in the blended region for a full day cycle.
 		//    
 		// The same applies to the backwards advancement of time case just in reverse. So instead of adding
 		// a BlendOffset of 1D, we subtract.
@@ -215,20 +224,8 @@ namespace UE::DaySequence
 			float BlendOffset = LastBlendOffset * DayLength;
 			const float StartTime = TargetTime + BlendDelta + BlendOffset;
 
-			// If we have applied a blend offset in the opposite direction of our LastBlendWinding and remain
-			// in the opposite space of the BlendWinding (ie. BlendDelta < 0.f for +Winding or BlendDelta > 0.f for -Winding)
-			// then we must invert the Lerp since the weight is mirrored across TargetTime.
-			if ((LastBlendWinding > 0 && (LastBlendOffset < 0 || (LastBlendOffset == 0 && BlendDelta < 0.f))) ||
-				(LastBlendWinding < 0 && (LastBlendOffset > 0 || (LastBlendOffset == 0 && BlendDelta > 0.f))))
-			{
-				AccumulatedTime = FMath::Lerp(TargetTime, StartTime, AccumulatedWeight);
-			}
-			else
-			{
-				AccumulatedTime = FMath::Lerp(StartTime, TargetTime, AccumulatedWeight);
-			}
-
 			// Wrap is inclusive of the upper bound. Ensure we wrap around correctly.
+			AccumulatedTime = FMath::Lerp(StartTime, TargetTime, AccumulatedWeight);
 			AccumulatedTime = FMath::Wrap(AccumulatedTime, 0.f, DayLength);
 			if (AccumulatedTime >= DayLength)
 			{
@@ -240,6 +237,10 @@ namespace UE::DaySequence
 			UE_LOG(LogDaySequence, VeryVerbose, TEXT("AccumTime: %f | AccumWeight: %f"), AccumulatedTime, AccumulatedWeight);
 #endif
 		}
+		else
+		{
+			ResetBlendState();
+		}
 		
 		// we could probably ensure(AccumulatedWeight is nearly equal to 1.f), leaving out for now
 
@@ -249,6 +250,15 @@ namespace UE::DaySequence
 		return AccumulatedTime;
 	}
 
+	void FStaticTimeManager::ResetBlendState() const
+	{
+		LastBlendDirection.Reset();
+		LastBlendDelta.Reset();
+		LastBlendWinding = 0;
+		LastBlendOffset = 0;
+	}
+
+
 	FStaticTimeInfo FStaticTimeManager::ProcessPriorityGroup(int32 StartIdx, int32 EndIdx) const
 	{
 		FStaticTimeInfo GroupInfo = {0.f, 0.f};
diff --git a/Engine/Plugins/Experimental/DaySequence/Source/DaySequence/Private/DaySequenceStaticTime.h b/Engine/Plugins/Experimental/DaySequence/Source/DaySequence/Private/DaySequenceStaticTime.h
--- a/Engine/Plugins/Experimental/DaySequence/Source/DaySequence/Private/DaySequenceStaticTime.h
+++ b/Engine/Plugins/Experimental/DaySequence/Source/DaySequence/Private/DaySequenceStaticTime.h
@@ -49,6 +49,7 @@ namespace UE::DaySequence
 		float GetStaticTime(float InitialTime, float DayLength) const;
 		
 	private:
+		void ResetBlendState() const;
 		
 		// Handles multiple contributors of the same priority by averaging their weights and times and returning a single desired weight and time
 		FStaticTimeInfo ProcessPriorityGroup(int32 StartIdx, int32 EndIdx) const;
