DaySequence: Add ability to force override update interval even if modifier weight is 0 or 1.

diff --git a/Engine/Plugins/Experimental/DaySequence/Source/DaySequence/Private/DaySequenceModifierComponent.cpp b/Engine/Plugins/Experimental/DaySequence/Source/DaySequence/Private/DaySequenceModifierComponent.cpp
--- a/Engine/Plugins/Experimental/DaySequence/Source/DaySequence/Private/DaySequenceModifierComponent.cpp
+++ b/Engine/Plugins/Experimental/DaySequence/Source/DaySequence/Private/DaySequenceModifierComponent.cpp
@@ -279,14 +279,16 @@ void UDaySequenceModifierComponent::DaySequenceUpdate()
 
 		if (OverrideUpdateIntervalHandle)
 		{
-			if (bIsEnabled && bSmoothBlending &&
-			FMath::IsWithin(GetBlendWeight(), 0.f + UE_KINDA_SMALL_NUMBER, 1.f - UE_KINDA_SMALL_NUMBER))
+			if (bIsEnabled && bSmoothBlending)
 			{
-				OverrideUpdateIntervalHandle->StartOverriding();
-			}
-			else
-			{
-				OverrideUpdateIntervalHandle->StopOverriding();
+				if (bForceSmoothBlending || FMath::IsWithin(GetBlendWeight(), 0.f + UE_KINDA_SMALL_NUMBER, 1.f - UE_KINDA_SMALL_NUMBER))
+				{
+					OverrideUpdateIntervalHandle->StartOverriding();
+				}
+				else
+				{
+					OverrideUpdateIntervalHandle->StopOverriding();
+				}
 			}
 		}
 	}
diff --git a/Engine/Plugins/Experimental/DaySequence/Source/DaySequence/Public/DaySequenceModifierComponent.h b/Engine/Plugins/Experimental/DaySequence/Source/DaySequence/Public/DaySequenceModifierComponent.h
--- a/Engine/Plugins/Experimental/DaySequence/Source/DaySequence/Public/DaySequenceModifierComponent.h
+++ b/Engine/Plugins/Experimental/DaySequence/Source/DaySequence/Public/DaySequenceModifierComponent.h
@@ -329,6 +329,10 @@ protected:
 	UPROPERTY(EditAnywhere, BlueprintReadWrite, AdvancedDisplay, Category="Day Sequence")
 	uint8 bSmoothBlending : 1;
 
+	/** If true, day sequence evaluation will be smooth regardless of blend weight. Note: Is always be very expensive! */
+	UPROPERTY(EditAnywhere, BlueprintReadWrite, AdvancedDisplay, Category="Day Sequence", meta=(EditCondition="bSmoothBlending", EditConditionHides))
+	uint8 bForceSmoothBlending : 1;
+
 private:
 	
 	/**
