DaySequence: Ignore ADaySequenceActor::Play() calls while bRunDayCycle is true.

A common pattern we have encountered is gameplay code attempting to get the day/night cycle into a known state on BeginPlay by starting playback and setting to some initial time. bRunDayCycle is meant to authoritatively determine whether or not the DaySequenceActor starts playback of the day cycle on BeginPlay but it gets overridden by this external game code on occasion. This change prevents external callers from beginning playback if bRunDayCycle is false.

diff --git a/Engine/Plugins/Experimental/DaySequence/Source/DaySequence/Private/DaySequenceActor.cpp b/Engine/Plugins/Experimental/DaySequence/Source/DaySequence/Private/DaySequenceActor.cpp
--- a/Engine/Plugins/Experimental/DaySequence/Source/DaySequence/Private/DaySequenceActor.cpp
+++ b/Engine/Plugins/Experimental/DaySequence/Source/DaySequence/Private/DaySequenceActor.cpp
@@ -1059,6 +1059,16 @@ void ADaySequenceActor::WarpEvaluationRange(FMovieSceneEvaluationRange& InOutRan
 	InOutRange.ResetRange(Range);
 }
 
+bool ADaySequenceActor::GetRunDayCycle() const
+{
+	return bRunDayCycle;
+}
+
+void ADaySequenceActor::SetRunDayCycle(bool bNewRunDayCycle)
+{
+	bRunDayCycle = bNewRunDayCycle;
+}
+
 float ADaySequenceActor::GetDayLength() const
 {
 	return DayLength.ToHours();
@@ -1250,6 +1260,12 @@ void ADaySequenceActor::UnregisterStaticTimeContributor(const UObject* InUserObj
 
 void ADaySequenceActor::Play()
 {
+	if (!bRunDayCycle)
+	{
+		UE_LOG(LogDaySequence, Warning, TEXT("Cannot begin/resume playback while bRunDayCycle is false."));
+		return;
+	}
+
 	UDaySequencePlayer* Player = GetSequencePlayerInternal();
 	UWorld* World = GetWorld();
 	if (HasValidRootSequence() && Player && World && World->IsGameWorld())
diff --git a/Engine/Plugins/Experimental/DaySequence/Source/DaySequence/Public/DaySequenceActor.h b/Engine/Plugins/Experimental/DaySequence/Source/DaySequence/Public/DaySequenceActor.h
--- a/Engine/Plugins/Experimental/DaySequence/Source/DaySequence/Public/DaySequenceActor.h
+++ b/Engine/Plugins/Experimental/DaySequence/Source/DaySequence/Public/DaySequenceActor.h
@@ -138,6 +138,18 @@ public:
 	UFUNCTION(BlueprintCallable, Category="TimeOfDay")
 	void SetTimeOfDayPreview(float InHours);
 
+	/**
+     * Get bRunDayCycle, a bool that determines whether or not the day cycle can be played back.
+     * @return bRunDayCycle
+     */
+    bool GetRunDayCycle() const;
+
+    /**
+     * Set bRunDayCycle, a bool that determines whether or not the day cycle can be played back.
+     * @param bNewRunDayCycle The new value for bRunDayCycle
+     */
+    void SetRunDayCycle(bool bNewRunDayCycle);
+	
 	/**
 	 * Get the length of each day in hours.
 	 *
@@ -489,7 +501,7 @@ protected:
 #endif
 
 	/** Whether or not to run a day cycle. If this is unchecked the day cycle will remain fixed at the time specified by the Initial Time setting */
-	UPROPERTY(EditAnywhere, Category = RuntimeDayCycle)
+	UPROPERTY(EditAnywhere, BlueprintReadWrite, Setter = SetRunDayCycle, Getter = GetRunDayCycle, Category = RuntimeDayCycle)
 	bool bRunDayCycle;
 
 	UPROPERTY(EditAnywhere, replicated, BlueprintReadWrite, Category = RuntimeDayCycle, meta=(InlineEditConditionToggle))
