DaySequence: Introduce GetApparentTimeOfDay(). This is distinct from GetTimeOfDay() in that it accounts for fixed time and time remapping via curve.

diff --git a/Engine/Plugins/Experimental/DaySequence/Source/DaySequence/Private/DaySequenceActor.cpp b/Engine/Plugins/Experimental/DaySequence/Source/DaySequence/Private/DaySequenceActor.cpp
--- a/Engine/Plugins/Experimental/DaySequence/Source/DaySequence/Private/DaySequenceActor.cpp
+++ b/Engine/Plugins/Experimental/DaySequence/Source/DaySequence/Private/DaySequenceActor.cpp
@@ -1200,6 +1200,14 @@ float ADaySequenceActor::GetTimeOfDay() const
 	return Result;
 }
 
+float ADaySequenceActor::GetApparentTimeOfDay() const
+{
+	float Time = HasStaticTimeOfDay() ? GetStaticTimeOfDay() : GetTimeOfDay();
+	const bool bUseCurve = bUseInterpCurve && DayInterpCurve && !bForceDisableDayInterpCurve;
+	Time = bUseCurve ? DayInterpCurve->FloatCurve.Eval(Time, Time) : Time;
+	return Time;
+}
+
 bool ADaySequenceActor::SetTimeOfDay(float InHours)
 {
 	// Only set time of day if we have a valid playing day sequence.
diff --git a/Engine/Plugins/Experimental/DaySequence/Source/DaySequence/Public/DaySequenceActor.h b/Engine/Plugins/Experimental/DaySequence/Source/DaySequence/Public/DaySequenceActor.h
--- a/Engine/Plugins/Experimental/DaySequence/Source/DaySequence/Public/DaySequenceActor.h
+++ b/Engine/Plugins/Experimental/DaySequence/Source/DaySequence/Public/DaySequenceActor.h
@@ -227,6 +227,15 @@ public:
 	UFUNCTION(BlueprintCallable, Category="TimeOfDay")
 	float GetTimeOfDay() const;
 
+	/**
+	 * Get the current apparent time of day in hours.
+	 * This is distinct from the actual time of day in that it is affected by static time and the day interp curve.
+	 *
+	 * @return float, time in hours
+	 */
+	UFUNCTION(BlueprintCallable, Category="TimeOfDay")
+	float GetApparentTimeOfDay() const;
+
 	/**
 	 * Set the current time of day in hours (server only).
 	 *
