Add a macro to disable the most expensive logging on DaySequencePlayer

#rnx

diff --git a/Engine/Plugins/Experimental/DaySequence/Source/DaySequence/Private/DaySequencePlayer.cpp b/Engine/Plugins/Experimental/DaySequence/Source/DaySequence/Private/DaySequencePlayer.cpp
--- a/Engine/Plugins/Experimental/DaySequence/Source/DaySequence/Private/DaySequencePlayer.cpp
+++ b/Engine/Plugins/Experimental/DaySequence/Source/DaySequence/Private/DaySequencePlayer.cpp
@@ -39,6 +39,11 @@ DECLARE_CYCLE_STAT(TEXT("Continue Evaluation"), DaySequenceEval_ContinueEvaluati
 DECLARE_CYCLE_STAT(TEXT("Synchronous Flush"), DaySequenceEval_SynchronousFlush, STATGROUP_DaySequenceEval);
 DECLARE_CYCLE_STAT(TEXT("Latent Action Flush"), DaySequenceEval_LatentActionFlush, STATGROUP_DaySequenceEval);
 
+// On lower end platforms logging these events can cause excessive stalling.
+#if !defined (UE_DAY_SEQUENCE_PLAYER_DISABLE_EXPENSIVE_LOGS)
+#define UE_DAY_SEQUENCE_PLAYER_DISABLE_EXPENSIVE_LOGS NO_LOGGING
+#endif
+
 namespace UE::DaySequence
 {
 	int32 GDaySequenceUpdateIntervalOverrideEnabled = 1;
@@ -1131,7 +1136,7 @@ void UDaySequencePlayer::UpdateMovieSceneInstance(FMovieSceneEvaluationRange InR
 		return;
 	}
 
-#if !NO_LOGGING
+#if !UE_DAY_SEQUENCE_PLAYER_DISABLE_EXPENSIVE_LOGS
 	if (UE_LOG_ACTIVE(LogDaySequence, VeryVerbose))
 	{
 		const FQualifiedFrameTime CurrentTime = GetCurrentTime();
@@ -1465,7 +1470,7 @@ void UDaySequencePlayer::RPC_ExplicitServerUpdateEvent_Implementation(EUpdatePos
 	// Explicit RPC call - empty our smoothed server samples
 	ServerTimeSamples.Reset();
 
-#if !NO_LOGGING
+#if !UE_DAY_SEQUENCE_PLAYER_DISABLE_EXPENSIVE_LOGS
 	// Log the sync event if necessary
 	if (UE_LOG_ACTIVE(LogDaySequence, Verbose))
 	{
@@ -1511,7 +1516,7 @@ void UDaySequencePlayer::RPC_OnStopEvent_Implementation(FFrameTime StoppedTime,
 	// Explicit RPC call - empty our smoothed server samples
 	ServerTimeSamples.Reset();
 
-#if !NO_LOGGING
+#if !UE_DAY_SEQUENCE_PLAYER_DISABLE_EXPENSIVE_LOGS
 	if (UE_LOG_ACTIVE(LogDaySequencePlayerRepl, Verbose))
 	{
 		const FFrameTime CurrentTime = PlayPosition.GetCurrentPosition();
@@ -1604,7 +1609,7 @@ void UDaySequencePlayer::PostNetReceive()
 	// these last few frames even though it has already stopped from the RPCs.
 	if (NetSyncProps.LastKnownSerialNumber < SerialNumber)
 	{
-#if !NO_LOGGING
+#if !UE_DAY_SEQUENCE_PLAYER_DISABLE_EXPENSIVE_LOGS
 		const FFrameTime CurrentTime = PlayPosition.GetCurrentPosition();
 		const FString SequenceName = GetSequenceName(true);
 		UE_LOG(LogDaySequencePlayerRepl, Verbose, 
@@ -1650,7 +1655,7 @@ void UDaySequencePlayer::PostNetReceive()
 	if (Observer && !Observer->CanObserveSequence())
 	{
 		// We shouldn't do anything.
-#if !NO_LOGGING
+#if !UE_DAY_SEQUENCE_PLAYER_DISABLE_EXPENSIVE_LOGS
 		if (UE_LOG_ACTIVE(LogDaySequencePlayerRepl, Verbose))
 		{
 			const FFrameTime CurrentTime = PlayPosition.GetCurrentPosition();
@@ -1664,7 +1669,7 @@ void UDaySequencePlayer::PostNetReceive()
 		return;
 	}
 
-#if !NO_LOGGING
+#if !UE_DAY_SEQUENCE_PLAYER_DISABLE_EXPENSIVE_LOGS
 	if (UE_LOG_ACTIVE(LogDaySequencePlayerRepl, VeryVerbose))
 	{
 		const FFrameTime CurrentTime = PlayPosition.GetCurrentPosition();
@@ -1810,7 +1815,7 @@ void UDaySequencePlayer::UpdateNetworkSync()
 		
 		if (Difference > LagThreshold + PingLag)
 		{
-#if !NO_LOGGING
+#if !UE_DAY_SEQUENCE_PLAYER_DISABLE_EXPENSIVE_LOGS
 			if (UE_LOG_ACTIVE(LogDaySequencePlayerRepl, Log))
 			{
 				const FFrameTime CurrentTime = PlayPosition.GetCurrentPosition();
