DaySequence: Add debug information on SetTimeOfDay

Rogue time setters cause endless headaches. We need a way to dump info about what code is setting time so we can take corrective action.

diff --git a/Engine/Plugins/Experimental/DaySequence/Source/DaySequence/Private/DaySequenceActor.cpp b/Engine/Plugins/Experimental/DaySequence/Source/DaySequence/Private/DaySequenceActor.cpp
--- a/Engine/Plugins/Experimental/DaySequence/Source/DaySequence/Private/DaySequenceActor.cpp
+++ b/Engine/Plugins/Experimental/DaySequence/Source/DaySequence/Private/DaySequenceActor.cpp
@@ -66,6 +66,14 @@ namespace UE::DaySequence
 		GFrameBudgetMicroseconds,
 		TEXT("(Default: 30us) Approximate max per-frame budget for time-of-day actors in microseconds.")
 	);
+
+	static TAutoConsoleVariable<bool> CVarDumpStackTraceOnSetTimeOfDay(
+	TEXT("DaySequence.DumpStackTraceOnSetTimeOfDay"),
+	false,
+	TEXT("When true, calls to SetTimeOfDay will dump the stack trace."),
+	ECVF_Default
+);
+	
 } // namespace UE::DaySequence
 
 ADaySequenceActor::ADaySequenceActor(const FObjectInitializer& Init)
@@ -1210,6 +1218,13 @@ float ADaySequenceActor::GetApparentTimeOfDay() const
 
 bool ADaySequenceActor::SetTimeOfDay(float InHours)
 {
+#if !UE_BUILD_SHIPPING
+	if (UE::DaySequence::CVarDumpStackTraceOnSetTimeOfDay.GetValueOnAnyThread())
+	{
+		FDebug::DumpStackTraceToLog(ELogVerbosity::Display);
+	}
+#endif
+	
 	// Only set time of day if we have a valid playing day sequence.
 	UWorld* World = GetWorld();
 	UDaySequencePlayer* Player = GetSequencePlayerInternal();
