DaySequence: Add Time of Day menu to the new viewport toolbar

diff --git a/Engine/Plugins/Experimental/DaySequence/Source/DaySequenceEditor/Private/DaySequenceEditorModule.cpp b/Engine/Plugins/Experimental/DaySequence/Source/DaySequenceEditor/Private/DaySequenceEditorModule.cpp
--- a/Engine/Plugins/Experimental/DaySequence/Source/DaySequenceEditor/Private/DaySequenceEditorModule.cpp
+++ b/Engine/Plugins/Experimental/DaySequence/Source/DaySequenceEditor/Private/DaySequenceEditorModule.cpp
@@ -52,6 +52,16 @@ static const FName DaySequenceTimeName("DaySequenceTime");
 static const FName DaySequenceViewportToolBarExtensionName("DaySequenceEditorViewportToolBar");
 static const FName EnvironmentLightingActorClassName("EnvironmentLightingActor");
 
+namespace UE::DaySequence
+{
+	static TAutoConsoleVariable<bool> CVarShowToolbarMenuLabel(
+		TEXT("DaySequence.ToolbarMenu.ShowLabel"),
+		true,
+		TEXT("When true, the Time of Day toolbar menu will have a label. When false, only an icon will be shown."),
+		ECVF_Default
+	);
+}
+
 class FDaySequenceEditorModule : public IDaySequenceEditorModule, public FGCObject
 {
 public:
@@ -104,6 +114,10 @@ private:
 
 	void CreateDaySequenceViewportMenu(UToolMenu* Menu);
 	TSharedRef<SWidget> CreateDaySequenceViewportToolbarExtension(FWeakObjectPtr ExtensionContext);
+
+	/** Shared widget functions */
+	static bool IsMenuVisible();
+	TSharedRef<SWidget> CreateTimeOfDayWidget() const;
 	
 	TSharedRef<SWidget> CreateDaySequencePreviewWidget();
 
@@ -535,6 +549,19 @@ void FDaySequenceEditorModule::RegisterMenus()
 				CreateDaySequenceViewportMenu(InMenu);
 			}));
 		}
+		
+		{
+			UToolMenu* ToolbarMenu = UToolMenus::Get()->ExtendMenu("LevelEditor.ViewportToolBar");
+			
+			ToolbarMenu->FindOrAddSection("Left").AddEntry(FToolMenuEntry::InitComboButton(
+				"TimeOfDay",
+				FToolUIActionChoice(FUIAction(FExecuteAction(), FCanExecuteAction(), FGetActionCheckState(), FIsActionButtonVisible::CreateLambda([]() { return IsMenuVisible(); }))),
+				FNewToolMenuChoice(FOnGetContent::CreateRaw(this, &FDaySequenceEditorModule::CreateTimeOfDayWidget)),
+				TAttribute<FText>::CreateLambda([]() { return UE::DaySequence::CVarShowToolbarMenuLabel.GetValueOnAnyThread() ? LOCTEXT("DaySequenceMenuLabel", "Time of Day") : FText::GetEmpty(); }),
+				FText::GetEmpty(),
+				FSlateIcon(FDaySequenceEditorStyle::GetStyleSetName(), "DaySequenceEditor.ViewportToolBar")
+			));
+		}
 	}
 }
 
@@ -559,7 +586,6 @@ void FDaySequenceEditorModule::CreateDaySequenceViewportMenu(UToolMenu* Menu)
 		FToolMenuSection& Section = Menu->AddSection("DaySequencePIESettings", LOCTEXT("DaySequencePIESettingsHeader", "PIE Settings"));
 		Section.AddEntry(FToolMenuEntry::InitMenuEntryWithCommandList(FDaySequenceEditorCommands::Get().OverrideInitialTimeOfDay, PluginCommands));
 		Section.AddEntry(FToolMenuEntry::InitMenuEntryWithCommandList(FDaySequenceEditorCommands::Get().OverrideRunDayCycle, PluginCommands));
-
 	}
 	{
 		FToolMenuSection& Section = Menu->AddSection("DaySequenceActions", LOCTEXT("DaySequenceActionsHeader", "Actions"));
@@ -570,39 +596,48 @@ void FDaySequenceEditorModule::CreateDaySequenceViewportMenu(UToolMenu* Menu)
 	}
 }
 
+bool FDaySequenceEditorModule::IsMenuVisible()
+{
+	bool Result = false;
+	if (!UE::GetIsEditorLoadingPackage() && GEditor)
+	{
+		if (const UWorld* World = GEditor->GetEditorWorldContext().World())
+		{
+			const UDaySequenceSubsystem* TODSubsystem = World->GetSubsystem<UDaySequenceSubsystem>();
+			Result = TODSubsystem && TODSubsystem->GetDaySequenceActor();
+		}
+	}
+	return Result;
+}
+
+TSharedRef<SWidget> FDaySequenceEditorModule::CreateTimeOfDayWidget() const
+{
+	static const FName MenuName("DaySequence.ViewportToolBar");
+	if (!UToolMenus::Get()->IsMenuRegistered(MenuName))
+	{
+		UToolMenus::Get()->RegisterMenu(MenuName);
+	}
+
+	const FToolMenuContext MenuContext(PluginCommands, TSharedPtr<FExtender>());
+	return UToolMenus::Get()->GenerateWidget(MenuName, MenuContext);
+}
+
 TSharedRef<SWidget> FDaySequenceEditorModule::CreateDaySequenceViewportToolbarExtension(FWeakObjectPtr ExtensionContext)
 {
 	if (const UViewportToolBarContext* ExtensionContextObject = Cast<UViewportToolBarContext>(ExtensionContext.Get()); ensure(ExtensionContextObject))
 	{
-		return
-			SNew(SEditorViewportToolbarMenu)
-				.ParentToolBar(ExtensionContextObject->ViewportToolBar.Pin())
-				.Label(LOCTEXT("DaySequenceMenuLabel", "Time of Day"))
-				.LabelIcon(FDaySequenceEditorStyle::Get()->GetBrush("DaySequenceEditor.ViewportToolBar"))
-				.OnGetMenuContent_Lambda([this]()
-				{
-					static const FName MenuName("DaySequence.ViewportToolBar");
-					if (!UToolMenus::Get()->IsMenuRegistered(MenuName))
-					{
-						UToolMenus::Get()->RegisterMenu(MenuName);
-					}
-
-					const FToolMenuContext MenuContext(PluginCommands, TSharedPtr<FExtender>());
-					return UToolMenus::Get()->GenerateWidget(MenuName, MenuContext);
-				})
-				.Visibility_Lambda([this]()
-				{
-					EVisibility Result = EVisibility::Collapsed;
-					if (!UE::GetIsEditorLoadingPackage() && GEditor)
-					{
-						if (const UWorld* World = GEditor->GetEditorWorldContext().World())
-						{
-							const UDaySequenceSubsystem* TODSubsystem = World->GetSubsystem<UDaySequenceSubsystem>();
-							Result = TODSubsystem && TODSubsystem->GetDaySequenceActor() ? EVisibility::Visible : EVisibility::Collapsed;
-						}
-					}
-					return Result;
-				});
+		return SNew(SEditorViewportToolbarMenu)
+			.ParentToolBar(ExtensionContextObject->ViewportToolBar.Pin())
+			.Label(LOCTEXT("DaySequenceMenuLabel", "Time of Day"))
+			.LabelIcon(FDaySequenceEditorStyle::Get()->GetBrush("DaySequenceEditor.ViewportToolBar"))
+			.OnGetMenuContent_Lambda([this]()
+			{
+				return CreateTimeOfDayWidget();
+			})
+			.Visibility_Lambda([this]()
+			{
+				return IsMenuVisible() ? EVisibility::Visible : EVisibility::Collapsed;
+			});
 	}
 	return SNullWidget::NullWidget;
 }
