DaySequenceModifierVolume: Stop forcing UserBlendWeight to 1.0 on blend target set and make UserBlendWeight non-transient.

diff --git a/Engine/Plugins/Experimental/DaySequence/Source/DaySequence/Private/Actors/DaySequenceModifierVolume.cpp b/Engine/Plugins/Experimental/DaySequence/Source/DaySequence/Private/Actors/DaySequenceModifierVolume.cpp
--- a/Engine/Plugins/Experimental/DaySequence/Source/DaySequence/Private/Actors/DaySequenceModifierVolume.cpp
+++ b/Engine/Plugins/Experimental/DaySequence/Source/DaySequence/Private/Actors/DaySequenceModifierVolume.cpp
@@ -41,7 +41,6 @@ void ADaySequenceModifierVolume::SetBlendTarget(APlayerController* InPC)
 	CurrentBlendTarget = InPC;
 
 	DaySequenceModifier->SetBlendTarget(InPC);
-	DaySequenceModifier->SetUserBlendWeight(1.f);
 }
 
 void ADaySequenceModifierVolume::BeginPlay()	
diff --git a/Engine/Plugins/Experimental/DaySequence/Source/DaySequence/Public/DaySequenceModifierComponent.h b/Engine/Plugins/Experimental/DaySequence/Source/DaySequence/Public/DaySequenceModifierComponent.h
--- a/Engine/Plugins/Experimental/DaySequence/Source/DaySequence/Public/DaySequenceModifierComponent.h
+++ b/Engine/Plugins/Experimental/DaySequence/Source/DaySequence/Public/DaySequenceModifierComponent.h
@@ -282,7 +282,7 @@ protected:
 	float BlendAmount;
 
 	/** User specified blend weight. The final blend weight is determined by BlendPolicy. */
-	UPROPERTY(Transient, EditAnywhere, BlueprintReadWrite, Setter, Getter, Category="Day Sequence", meta=(DisplayAfter="BlendPolicy", EditCondition="BlendPolicy!=EDaySequenceModifierUserBlendPolicy::Ignored", EditConditionHides))
+	UPROPERTY(EditAnywhere, BlueprintReadWrite, Setter, Getter, Category="Day Sequence", meta=(DisplayAfter="BlendPolicy", EditCondition="BlendPolicy!=EDaySequenceModifierUserBlendPolicy::Ignored", EditConditionHides))
 	float UserBlendWeight;
 
 	/** Changes the way the modifier controls the day/night cycle time when enabled. */
