DaySequenceActor: Fix issue for paused DSAs when SequenceUpdateInterval is 0.0

StartDaySequenceUpdateTimer() begins a timer when the sequence player is paused. The goal is to always have something ticking at roughly the same rate we expect the player to update at. When SequenceUpdateInterval is 0, this means we try to specify the timer rate as 0 which actually results in the timer never doing anything. If the SequenceUpdateInterval is 0 lets just use UE_SMALL_NUMBER as the rate.

diff --git a/Engine/Plugins/Experimental/DaySequence/Source/DaySequence/Private/DaySequenceActor.cpp b/Engine/Plugins/Experimental/DaySequence/Source/DaySequence/Private/DaySequenceActor.cpp
--- a/Engine/Plugins/Experimental/DaySequence/Source/DaySequence/Private/DaySequenceActor.cpp
+++ b/Engine/Plugins/Experimental/DaySequence/Source/DaySequence/Private/DaySequenceActor.cpp
@@ -897,7 +897,7 @@ void ADaySequenceActor::SequencePlayerUpdated(float CurrentTime, float PreviousT
 
 void ADaySequenceActor::StartDaySequenceUpdateTimer()
 {
-	if (HasAuthority())
+	if (GetNetMode() == NM_DedicatedServer)
 	{
 		return;
 	}
@@ -911,13 +911,13 @@ void ADaySequenceActor::StartDaySequenceUpdateTimer()
 		World->GetTimerManager().SetTimer(DaySequenceUpdateTimerHandle, FTimerDelegate::CreateWeakLambda(this, [this]()
 		{
 			OnDaySequenceUpdate.Broadcast();
-		}), SequenceUpdateInterval, TimerParameters);
+		}), FMath::Max(SequenceUpdateInterval, UE_SMALL_NUMBER), TimerParameters);
 	}
 }
 
 void ADaySequenceActor::StopDaySequenceUpdateTimer()
 {
-	if (HasAuthority())
+	if (GetNetMode() == NM_DedicatedServer)
 	{
 		return;
 	}
