DaySequence: Add experimental splitscreen support for DaySequenceModifierVolumes.

This change:

1) Adds the concept of a 'specialized binding' for day sequence binding references. Historically we have interpreted an empty binding reference as implicitly referencing the current day sequence actor. Specialized bindings extend this pattern to allow empty bindings (more specifically, a binding which does not have a path to an object) to resolve based on its specialization. This change introduces 2 specializations: Root and CameraModifier. The Root specialization represents the original behavior that results in a reference to the active DSA. The CameraModifier specialization is a type of binding that is intended to resolve to a UDaySequenceCameraModifier. This modifier is spawned at runtime and the resolution depends on the day sequence modifier component providing a resolve function for this binding type. In the future I would like the DSA to be able to resolve this as well for _all_ local players' camera modifiers, but this has several problems that need to be address first and we shouldn't conflate those problems with this splitscreen support pass.

2) Extends the "Add" menu in the sequence editor for Day Sequences by creating a submenu called Specialized Bindings which allow users to add and remove these specialized bindings at will. For the time being, only 1 of each type of specialized binding can exist at a time.

3) Adds UDaySequenceCameraModifier. This is a type of camera modifier with a sequencer-exposed post process settings struct which can be targeted by sequences using the aforementioned special binding and which automatically feeds post process settings into the camera modifier stack.

4) Adds UDaySequenceCameraModifierManager. This manages the creation of our camera modifiers, a mapping of player controllers to the appropriate modifier, an editor world camera modifier for when there is no player controller, and an editor world post process component for editor preview. This is owned by the Day Sequence Actor. This primarily exists to encapsulate camera modifier related functions and can be thought of as an extension to the DSA.

5) Modifies ADaySequenceModifierVolume to support multiple local players. By default the modifier volumes comes with a single modifier component, but when additional local players are detected it will duplicate this modifier component. The goal is to have 1 modifier component per local player. From a configuration perspective, nothing has changed. These details are hidden from the user.

6) Adds a per-subsequence binding override scheme. While the Root specialization can be resolved by the binding reference itself, the CameraModifier specialization cannot be because it is intended to actually resolve different for different instances of the same sequence in the root sequence. This means 2 modifiers with different blend targets that initialize a subsequence from the same source sequence will ultimately result in 2 different camera modifiers being animated. This override scheme is intended to be general purpose for future use cases like this.

7) Modifies UDaySequenceModifierComponent to take advantage of the per-subsequence binding override scheme by registering resolve functions when initializing a subsequence with a CameraModifier binding.

diff --git a/Engine/Plugins/Experimental/DaySequence/Source/DaySequence/Private/Actors/DaySequenceModifierVolume.cpp b/Engine/Plugins/Experimental/DaySequence/Source/DaySequence/Private/Actors/DaySequenceModifierVolume.cpp
--- a/Engine/Plugins/Experimental/DaySequence/Source/DaySequence/Private/Actors/DaySequenceModifierVolume.cpp
+++ b/Engine/Plugins/Experimental/DaySequence/Source/DaySequence/Private/Actors/DaySequenceModifierVolume.cpp
@@ -8,13 +8,22 @@
 
 #include "Components/BoxComponent.h"
 #include "Engine/World.h"
-#include "GameDelegates.h"
 #include "GameFramework/Pawn.h"
-#include "TimerManager.h"
 #include "Net/UnrealNetwork.h"
 
+namespace UE::DaySequence
+{
+	static TAutoConsoleVariable<bool> CVarModifierVolumeEnableSplitscreenSupport(
+	TEXT("DaySequence.ModifierVolume.EnableSplitscreenSupport"),
+	true,
+	TEXT("When true, Day Sequence Modifier Volumes attempt to initialize transient modifier components for all local players."),
+	ECVF_Default
+	);
+}
+
 ADaySequenceModifierVolume::ADaySequenceModifierVolume(const FObjectInitializer& Init)
 : Super(Init)
+, bEnableSplitscreenSupport(false)
 {
 	PrimaryActorTick.bCanEverTick = true;
 	
@@ -31,34 +40,18 @@ ADaySequenceModifierVolume::ADaySequenceModifierVolume(const FObjectInitializer&
 	DaySequenceModifier->AddVolumeShapeComponent(DefaultBoxReference);
 }
 
-void ADaySequenceModifierVolume::SetBlendTarget(APlayerController* InPC)
-{
-	if (!IsValid(InPC) || InPC == CurrentBlendTarget)
-	{
-		return;
-	}
-
-	CurrentBlendTarget = InPC;
-
-	DaySequenceModifier->SetBlendTarget(InPC);
-}
-
 void ADaySequenceModifierVolume::BeginPlay()	
 {
 	Super::BeginPlay();
 
 	Initialize();
+}
 
-	if (const UWorld* World = GetWorld(); World && World->IsPlayingReplay())
-	{
-		ReplayScrubbedHandle = FNetworkReplayDelegates::OnReplayScrubComplete.AddWeakLambda(this, [this](const UWorld* InWorld)
-		{
-			if (InWorld == GetWorld())
-			{
-				DaySequenceActorSetup();
-			}
-		});
-	}
+void ADaySequenceModifierVolume::EndPlay(const EEndPlayReason::Type EndPlayReason)
+{
+	Deinitialize();
+	
+	Super::EndPlay(EndPlayReason);
 }
 
 void ADaySequenceModifierVolume::OnConstruction(const FTransform& Transform)
@@ -82,75 +75,77 @@ void ADaySequenceModifierVolume::Initialize()
 		return;
 	}
 
-#if WITH_EDITOR
-	const UWorld* World = GetWorld();
-	if (World && World->WorldType == EWorldType::Editor)
+	if (const UWorld* World = GetWorld())
 	{
-		DaySequenceActor = nullptr;
-		if (IsValid(DaySequenceModifier))
+#if WITH_EDITOR
+		if (World->WorldType == EWorldType::Editor)
 		{
-			DaySequenceModifier->UnbindFromDaySequenceActor();
+			DaySequenceActor = nullptr;
+			if (IsValid(DaySequenceModifier))
+			{
+				DaySequenceModifier->UnbindFromDaySequenceActor();
+			}
 		}
-	}
 #endif
 
-	DaySequenceActorSetup();
-}
-
-void ADaySequenceModifierVolume::PlayerControllerSetup()
-{
-#if WITH_EDITOR
-	UWorld* World = GetWorld();
-	if (World && World->WorldType != EWorldType::Editor)
-	{
-#endif
-	CachePlayerController();
-#if WITH_EDITOR
-	}
-#endif
-}
-
-void ADaySequenceModifierVolume::CachePlayerController()
-{
-	if (const UWorld* World = GetWorld())
-	{
-		for (FConstPlayerControllerIterator Iterator = World->GetPlayerControllerIterator(); Iterator; ++Iterator)
+		DaySequenceActorSetup();
+		
+		if (World->IsGameWorld())
 		{
-			if (APlayerController* PlayerController = Iterator->Get())
+			auto HandleNewPlayerController = [this](APlayerController* PlayerController)
 			{
-				if (PlayerController->IsLocalPlayerController())
+				if (PlayerController->IsLocalController())
 				{
-					CachedPlayerController = PlayerController;
-					break;
+					CreatePlayer(PlayerController);
+				}
+			};
+
+			// Create players that exist initially. We may not find anything here, but the actor spawned handler below should handle anything we miss.
+			for (FConstPlayerControllerIterator Iterator = World->GetPlayerControllerIterator(); Iterator; ++Iterator)
+			{
+				if (APlayerController* PlayerController = Iterator->Get())
+				{
+					HandleNewPlayerController(PlayerController);
 				}
 			}
+
+			ActorSpawnedHandle = World->AddOnActorSpawnedHandler(FOnActorSpawned::FDelegate::CreateWeakLambda(this, [HandleNewPlayerController](AActor* SpawnedActor)
+			{
+				if (APlayerController* PlayerController = Cast<APlayerController>(SpawnedActor))
+				{
+					HandleNewPlayerController(PlayerController);
+				}
+			}));
 		}
-	}
 
-	// If no local player controller found, queue this function for next tick.
-	if (!IsValid(CachedPlayerController))
-	{
-		QueuePlayerControllerQuery();
-	}
-	else
-	{
-		SetBlendTarget(CachedPlayerController);
+		if (World->IsPlayingReplay())
+		{
+			ReplayScrubbedHandle = FNetworkReplayDelegates::OnReplayScrubComplete.AddWeakLambda(this, [this](const UWorld* InWorld)
+			{
+				if (InWorld == GetWorld())
+				{
+					DaySequenceActorSetup();
+				}
+			});
+		}
 	}
 }
 
-void ADaySequenceModifierVolume::QueuePlayerControllerQuery()
+void ADaySequenceModifierVolume::Deinitialize()
 {
-	if (!IsValid(this))
-	{
-		return;
-	}
-
 	if (const UWorld* World = GetWorld())
 	{
-		World->GetTimerManager().SetTimerForNextTick([this]()
+		if (ActorSpawnedHandle.IsValid())
 		{
-			CachePlayerController();
-		});
+			World->RemoveOnActorSpawnedHandler(ActorSpawnedHandle);
+			ActorSpawnedHandle.Reset();
+		}
+		
+		if (ReplayScrubbedHandle.IsValid())
+		{
+			FNetworkReplayDelegates::OnReplayScrubComplete.Remove(ReplayScrubbedHandle);
+			ReplayScrubbedHandle.Reset();
+		}
 	}
 }
 
@@ -162,7 +157,17 @@ void ADaySequenceModifierVolume::DaySequenceActorSetup()
 
 void ADaySequenceModifierVolume::BindToDaySequenceActor()
 {
-	bool bNewActorSet = false;
+	auto RebindAllModifiers = [this]()
+	{
+		DaySequenceModifier->BindToDaySequenceActor(DaySequenceActor);
+		for (auto AdditionalPlayerIterator = AdditionalPlayers.CreateIterator(); AdditionalPlayerIterator; ++AdditionalPlayerIterator)
+		{
+			UDaySequenceModifierComponent* ModifierComponent = AdditionalPlayerIterator->Value;
+			ModifierComponent->BindToDaySequenceActor(DaySequenceActor);
+		}
+
+		OnDaySequenceActorBound(DaySequenceActor);
+	};
 	
 	if (const UWorld* World = GetWorld())
 	{
@@ -173,20 +178,11 @@ void ADaySequenceModifierVolume::BindToDaySequenceActor()
 				if (NewActor != DaySequenceActor)
 				{
 					DaySequenceActor = NewActor;
-					bNewActorSet = true;
+					RebindAllModifiers();
 				}
 			}
 		}
 	}
-	
-	if (bNewActorSet)
-	{
-		DaySequenceModifier->BindToDaySequenceActor(DaySequenceActor);
-
-		PlayerControllerSetup();
-
-		OnDaySequenceActorBound(DaySequenceActor);
-	}
 }
 
 void ADaySequenceModifierVolume::SetupDaySequenceSubsystemCallbacks()
@@ -205,4 +201,49 @@ void ADaySequenceModifierVolume::SetupDaySequenceSubsystemCallbacks()
 			}
 		}
 	}
+}
+
+void ADaySequenceModifierVolume::CreatePlayer(APlayerController* InPC)
+{
+	if (CachedPlayerController == InPC || AdditionalPlayers.Find(InPC))
+	{
+		return;
+	}
+
+	UDaySequenceModifierComponent* PlayerModifier;
+	
+	if (CachedPlayerController == nullptr)
+	{
+		CachedPlayerController = InPC;
+		PlayerModifier = DaySequenceModifier;
+	}
+	else if (IsSplitscreenSupported())
+	{
+        PlayerModifier = DuplicateObject<UDaySequenceModifierComponent>(DaySequenceModifier, this, TEXT("AdditionalPlayerModifier"));
+		AdditionalPlayers.FindOrAdd(InPC) = PlayerModifier;
+	}
+	else
+	{
+		// If CachedPlayerController is not null (points to the first local player) and we don't support splitscreen, just bail out.
+		return;
+	}
+
+	if (ensureMsgf(PlayerModifier, TEXT("PlayerModifier is nullptr!")))
+	{
+		if (!PlayerModifier->IsRegistered())
+		{
+			// This happens for duplicated modifiers.
+			PlayerModifier->RegisterComponent();
+			PlayerModifier->AttachToComponent(RootComponent, FAttachmentTransformRules::KeepRelativeTransform);
+			PlayerModifier->SetBias(DaySequenceModifier->GetBias() + AdditionalPlayers.Num());
+		}
+
+		PlayerModifier->SetBlendTarget(InPC);
+		PlayerModifier->BindToDaySequenceActor(DaySequenceActor);
+	}
+}
+
+bool ADaySequenceModifierVolume::IsSplitscreenSupported() const
+{
+	return bEnableSplitscreenSupport && UE::DaySequence::CVarModifierVolumeEnableSplitscreenSupport.GetValueOnAnyThread();
 }
\ No newline at end of file
diff --git a/Engine/Plugins/Experimental/DaySequence/Source/DaySequence/Private/DaySequence.cpp b/Engine/Plugins/Experimental/DaySequence/Source/DaySequence/Private/DaySequence.cpp
--- a/Engine/Plugins/Experimental/DaySequence/Source/DaySequence/Private/DaySequence.cpp
+++ b/Engine/Plugins/Experimental/DaySequence/Source/DaySequence/Private/DaySequence.cpp
@@ -5,6 +5,7 @@
 #include "DaySequencePlayer.h"
 #include "DaySequenceTrack.h"
 #include "DaySequenceModule.h"
+#include "DaySequenceCameraModifier.h"
 
 #include "Animation/AnimInstance.h"
 #include "Components/ActorComponent.h"
@@ -36,6 +37,8 @@
 
 #include UE_INLINE_GENERATED_CPP_BY_NAME(DaySequence)
 
+#define LOCTEXT_NAMESPACE "DaySequence"
+
 #if WITH_EDITOR
 UDaySequence::FPostDuplicateEvent UDaySequence::PostDuplicateEvent;
 #endif
@@ -107,6 +110,69 @@ void UDaySequence::AddDefaultBinding(const FGuid& PossessableGuid)
 	BindingReferences.AddDefaultBinding(PossessableGuid);
 }
 
+void UDaySequence::AddSpecializedBinding(EDaySequenceBindingReferenceSpecialization Specialization)
+{
+	if (Specialization == EDaySequenceBindingReferenceSpecialization::None)
+	{
+		return;
+	}
+
+	auto GetBindingName = [Specialization]() -> FString
+	{
+		static const FString RootBindingName = LOCTEXT("RootBindingName", "Root Day Sequence Actor").ToString();
+		static const FString CameraModifierBindingName = LOCTEXT("CameraModifierBindingName", "Day Sequence Camera Modifier").ToString();
+		
+		switch (Specialization)
+		{
+		case EDaySequenceBindingReferenceSpecialization::Root:
+			return RootBindingName;
+			
+		case EDaySequenceBindingReferenceSpecialization::CameraModifier:
+			return CameraModifierBindingName;
+			
+		case EDaySequenceBindingReferenceSpecialization::None:
+		default:
+			ensureMsgf(true, TEXT("Invalid specialization provided, assuming Root specialization!"));
+			return RootBindingName;
+		}
+	};
+	
+	auto GetClass = [Specialization]() -> UClass*
+	{
+		switch (Specialization)
+		{
+		case EDaySequenceBindingReferenceSpecialization::Root:
+			return ADaySequenceActor::StaticClass();
+			
+		case EDaySequenceBindingReferenceSpecialization::CameraModifier:
+			return UDaySequenceCameraModifier::StaticClass();
+			
+		case EDaySequenceBindingReferenceSpecialization::None:
+		default:
+			ensureMsgf(true, TEXT("Invalid specialization provided, assuming Root specialization!"));
+			return ADaySequenceActor::StaticClass();
+		}
+	};
+	
+	// only 1 allowed
+	if (BindingReferences.FindSpecializedBinding(Specialization).IsValid())
+	{
+		return;
+	}
+	
+	// Add a default binding
+	FMovieScenePossessable Possessable(GetBindingName(), GetClass());
+	FMovieSceneBinding     Binding(Possessable.GetGuid(), GetBindingName());
+
+	BindingReferences.AddSpecializedBinding(Possessable.GetGuid(), Specialization);
+	MovieScene->AddPossessable(Possessable, Binding);
+}
+
+FGuid UDaySequence::GetSpecializedBinding(EDaySequenceBindingReferenceSpecialization Specialization) const
+{
+	return BindingReferences.FindSpecializedBinding(Specialization);
+}
+
 UObject* UDaySequence::MakeSpawnableTemplateFromInstance(UObject& InSourceObject, FName ObjectName)
 {
 	return MovieSceneHelpers::MakeSpawnableTemplateFromInstance(InSourceObject, MovieScene, ObjectName);
@@ -624,3 +690,5 @@ const TArray<UAssetUserData*>* UDaySequence::GetAssetUserDataArray() const
 {
 	return &ToRawPtrTArrayUnsafe(AssetUserData);
 }
+
+#undef LOCTEXT_NAMESPACE
\ No newline at end of file
diff --git a/Engine/Plugins/Experimental/DaySequence/Source/DaySequence/Private/DaySequenceActor.cpp b/Engine/Plugins/Experimental/DaySequence/Source/DaySequence/Private/DaySequenceActor.cpp
--- a/Engine/Plugins/Experimental/DaySequence/Source/DaySequence/Private/DaySequenceActor.cpp
+++ b/Engine/Plugins/Experimental/DaySequence/Source/DaySequence/Private/DaySequenceActor.cpp
@@ -1,6 +1,7 @@
 // Copyright Epic Games, Inc. All Rights Reserved.
 
 #include "DaySequenceActor.h"
+#include "DaySequenceCameraModifier.h"
 #include "DaySequenceCollectionAsset.h"
 #include "DaySequence.h"
 #include "DaySequenceModule.h"
@@ -124,6 +125,7 @@ ADaySequenceActor::ADaySequenceActor(const FObjectInitializer& Init)
 	// SequencePlayer must be a default sub object for it to be replicated correctly
 	SequencePlayer = Init.CreateDefaultSubobject<UDaySequencePlayer>(this, "AnimationPlayer", true);
 	BindingOverrides = Init.CreateDefaultSubobject<UMovieSceneBindingOverrides>(this, "BindingOverrides");
+	CameraModifierManager = Init.CreateDefaultSubobject<UDaySequenceCameraModifierManager>(this, "CameraModifierManager");
 
 	bAlwaysRelevant = true;
 	bReplicates = true;
@@ -200,6 +202,13 @@ void ADaySequenceActor::Tick(float DeltaTime)
 	{
 		Player->Tick(DeltaTime);
 	}
+
+#if WITH_EDITOR
+	if (CameraModifierManager)
+	{
+		CameraModifierManager->UpdateEditorPreview();
+	}
+#endif
 }
 
 bool ADaySequenceActor::ShouldTickIfViewportsOnly() const
@@ -369,6 +378,14 @@ void ADaySequenceActor::Destroyed()
 
 bool ADaySequenceActor::RetrieveBindingOverrides(const FGuid& InBindingId, FMovieSceneSequenceID InSequenceID, TArray<UObject*, TInlineAllocator<1>>& OutObjects) const
 {
+	if (const TMap<FGuid, TFunction<bool(TArray<UObject*, TInlineAllocator<1>>&)>>* Entry = BindingResolveFunctions.Find(InSequenceID))
+	{
+		if (const TFunction<bool(TArray<UObject*, TInlineAllocator<1>>&)>* ResolveFunc = Entry->Find(InBindingId))
+		{
+			return (*ResolveFunc)(OutObjects);
+		}
+	}
+	
 	// todo: make this an ensureMsgf? This is null when duping a DSA BP.
 	if (BindingOverrides)
 	{
@@ -457,6 +474,11 @@ void ADaySequenceActor::PostEditChangeProperty(FPropertyChangedEvent& PropertyCh
 		ReinitializeActor();
 	}
 
+	if (CameraModifierManager)
+	{
+		CameraModifierManager->ResetEditorPreview();
+	}
+	
 	Super::PostEditChangeProperty(PropertyChangedEvent);
 }
 
@@ -603,6 +625,8 @@ void ADaySequenceActor::InitializeRootSequence()
 		{
 			if (!Section.Value)
 			{
+				BindingResolveFunctions.Remove(Section.Key->GetSequenceID());
+				
 #if WITH_EDITOR
 				OnSubSectionRemovedEvent.Broadcast(Section.Key);
 #endif
@@ -619,10 +643,12 @@ void ADaySequenceActor::InitializeRootSequence()
 	}
 	else
 	{
+		BindingResolveFunctions.Empty();
 		InitializeDaySequences();
 		OnPostInitializeDaySequences.Broadcast(nullptr);
 	}
 #else
+	BindingResolveFunctions.Empty();
 	InitializeDaySequences();
 	OnPostInitializeDaySequences.Broadcast(nullptr);
 #endif
@@ -1281,6 +1307,31 @@ void ADaySequenceActor::UnregisterStaticTimeContributor(const UObject* InUserObj
 	StaticTimeManager->RemoveStaticTimeContributor(InUserObject);
 }
 
+void ADaySequenceActor::RegisterBindingResolveFunction(FMovieSceneSequenceID SequenceID, FGuid Guid, TFunction<bool(TArray<UObject*, TInlineAllocator<1>>&)> InFunction)
+{
+	UnregisterBindingResolveFunction(SequenceID, Guid);
+
+	BindingResolveFunctions.FindOrAdd(SequenceID).FindOrAdd(Guid) = InFunction;
+}
+
+void ADaySequenceActor::UnregisterBindingResolveFunction(FMovieSceneSequenceID SequenceID, FGuid Guid)
+{
+	if (auto* SubsequenceResolveFunctions = BindingResolveFunctions.Find(SequenceID))
+	{
+		if (Guid.IsValid())
+		{
+			SubsequenceResolveFunctions->Remove(Guid);
+		}
+
+		// If no more functions associated with this sequence ID or no binding guid was specified, just remove the entry for this sequence.
+		if (SubsequenceResolveFunctions->Num() == 0 || !Guid.IsValid())
+		{
+			SubsequenceResolveFunctions = nullptr;	// we are invalidating this ptr
+			BindingResolveFunctions.Remove(SequenceID);
+		}
+	}
+}
+
 void ADaySequenceActor::Play()
 {
 	if (!bRunDayCycle)
diff --git a/Engine/Plugins/Experimental/DaySequence/Source/DaySequence/Private/DaySequenceBindingReference.cpp b/Engine/Plugins/Experimental/DaySequence/Source/DaySequence/Private/DaySequenceBindingReference.cpp
--- a/Engine/Plugins/Experimental/DaySequence/Source/DaySequence/Private/DaySequenceBindingReference.cpp
+++ b/Engine/Plugins/Experimental/DaySequence/Source/DaySequence/Private/DaySequenceBindingReference.cpp
@@ -53,8 +53,14 @@ FDaySequenceBindingReference::FDaySequenceBindingReference(UObject* InObject, UO
 }
 
 FDaySequenceBindingReference FDaySequenceBindingReference::DefaultRootBinding()
+{
+	return SpecializedBinding(EDaySequenceBindingReferenceSpecialization::Root);
+}
+
+FDaySequenceBindingReference FDaySequenceBindingReference::SpecializedBinding(EDaySequenceBindingReferenceSpecialization Specialization)
 {
 	FDaySequenceBindingReference Default;
+	Default.Specialization = Specialization;
 	return Default;
 }
 
@@ -68,12 +74,24 @@ UObject* FDaySequenceBindingReference::Resolve(UObject* InContext) const
 	// Context must always be an AActor - either the ADaySequenceActor playing back, or a parent actor (for component bindings)
 	check(InContext->IsA<AActor>());
 
-	// Empty binding is used to quickly identify the global Day Sequence Actor binding
+	// This branch is where specialized bindings are resolved (or where resolution is deferred).
 	if (ExternalObjectPath.IsNull() && ObjectPath.Len() == 0)
 	{
-		ADaySequenceActor* DaySequenceActor = Cast<ADaySequenceActor>(InContext);
-		ensureMsgf(DaySequenceActor, TEXT("Failed to locate the currently playing back day sequence actor - was InContext provided incorrectly?"));
-		return DaySequenceActor;
+		switch(Specialization)
+		{
+		case EDaySequenceBindingReferenceSpecialization::None:	// legacy path here, new references should never be None if the object paths are empty.
+		case EDaySequenceBindingReferenceSpecialization::Root:
+		{
+			ADaySequenceActor* DaySequenceActor = Cast<ADaySequenceActor>(InContext);
+			ensureMsgf(DaySequenceActor, TEXT("Failed to locate the currently playing back day sequence actor - was InContext provided incorrectly?"));
+			return DaySequenceActor;
+		}
+
+		// resolution handled elsewhere, not enough info
+		default:
+		case EDaySequenceBindingReferenceSpecialization::CameraModifier:
+			return nullptr;
+		}
 	}
 
 	if (!ExternalObjectPath.IsNull())
@@ -125,6 +143,12 @@ void FDaySequenceBindingReference::PerformLegacyFixup()
 		ObjectPath.Empty();
 		ExternalObjectPath.Reset();
 	}
+
+	// Fixup the specialization of empty bindings that have no specialization. This is the case for Root DSA bindings created before the specialization system.
+	if (ExternalObjectPath.IsNull() && ObjectPath.Len() == 0 && Specialization == EDaySequenceBindingReferenceSpecialization::None)
+	{
+		Specialization = EDaySequenceBindingReferenceSpecialization::Root;
+	}
 }
 #endif
 
@@ -147,13 +171,33 @@ void FDaySequenceBindingReferences::AddBinding(const FGuid& ObjectId, UObject* I
 
 void FDaySequenceBindingReferences::AddDefaultBinding(const FGuid& ObjectId)
 {
-	BindingIdToReferences.FindOrAdd(ObjectId).References.Emplace(FDaySequenceBindingReference::DefaultRootBinding());
+	AddSpecializedBinding(ObjectId, EDaySequenceBindingReferenceSpecialization::Root);
+}
+
+void FDaySequenceBindingReferences::AddSpecializedBinding(const FGuid& ObjectId, EDaySequenceBindingReferenceSpecialization Specialization)
+{
+	BindingIdToReferences.FindOrAdd(ObjectId).References.Emplace(FDaySequenceBindingReference::SpecializedBinding(Specialization));
+	SpecializedReferenceToGuid.FindOrAdd(Specialization) = ObjectId;
+	GuidToSpecializedReference.FindOrAdd(ObjectId) = Specialization;
+}
+
+FGuid FDaySequenceBindingReferences::FindSpecializedBinding(EDaySequenceBindingReferenceSpecialization Specialization) const
+{
+	const FGuid* BindingGuid = SpecializedReferenceToGuid.Find(Specialization);
+	return BindingGuid ? *BindingGuid : FGuid();
 }
 
 void FDaySequenceBindingReferences::RemoveBinding(const FGuid& ObjectId)
 {
 	BindingIdToReferences.Remove(ObjectId);
 	AnimSequenceInstances.Remove(ObjectId);
+
+	// If this Guid is a specialized binding, we need to remove our tracking info.
+	if (const EDaySequenceBindingReferenceSpecialization* Specialization = GuidToSpecializedReference.Find(ObjectId))
+	{
+		SpecializedReferenceToGuid.Remove(*Specialization);
+		GuidToSpecializedReference.Remove(ObjectId);	// warning: Specialization is probably a bad ptr after this removal.
+	}
 }
 
 void FDaySequenceBindingReferences::RemoveObjects(const FGuid& ObjectId, const TArray<UObject*>& InObjects, UObject* InContext)
@@ -243,7 +287,16 @@ void FDaySequenceBindingReferences::PerformLegacyFixup()
 	{
 		for (FDaySequenceBindingReference& Reference : Pair.Value.References)
 		{
+			const EDaySequenceBindingReferenceSpecialization PreFixupSpecialization = Reference.GetSpecialization();
 			Reference.PerformLegacyFixup();
+			const EDaySequenceBindingReferenceSpecialization PostFixupSpecialization = Reference.GetSpecialization();
+
+			// Make sure we track already created bindings that are fixed up to be specialized (this is the case for old Root DSA bindings).
+			if (PreFixupSpecialization == EDaySequenceBindingReferenceSpecialization::None && PreFixupSpecialization != PostFixupSpecialization)
+			{
+				SpecializedReferenceToGuid.FindOrAdd(PostFixupSpecialization) = Pair.Key;
+				GuidToSpecializedReference.FindOrAdd(Pair.Key) = PostFixupSpecialization;
+			}
 		}
 	}
 }
diff --git a/Engine/Plugins/Experimental/DaySequence/Source/DaySequence/Private/DaySequenceCameraModifier.cpp b/Engine/Plugins/Experimental/DaySequence/Source/DaySequence/Private/DaySequenceCameraModifier.cpp
new file mode 100644
--- /dev/null
+++ b/Engine/Plugins/Experimental/DaySequence/Source/DaySequence/Private/DaySequenceCameraModifier.cpp
@@ -0,0 +1,67 @@
+// Copyright Epic Games, Inc. All Rights Reserved.
+
+#include "DaySequenceCameraModifier.h"
+
+#include "Components/PostProcessComponent.h"
+#include "Engine/World.h"
+#include "GameFramework/PlayerController.h"
+
+#include UE_INLINE_GENERATED_CPP_BY_NAME(DaySequenceCameraModifier)
+
+UCameraModifier* UDaySequenceCameraModifierManager::GetCameraModifier(APlayerController* InPC)
+{
+	if (InPC)
+	{
+		if (const TWeakObjectPtr<UCameraModifier>* Modifier = CameraModifiers.Find(InPC))
+		{
+			return Modifier->Get();
+		}
+
+		if (APlayerCameraManager* CameraManager = InPC->PlayerCameraManager)
+		{
+			UCameraModifier* NewModifier = CameraManager->AddNewCameraModifier(UDaySequenceCameraModifier::StaticClass());
+			CameraModifiers.Emplace(InPC, NewModifier);
+			return NewModifier;
+		}
+	}
+#if WITH_EDITOR
+	else
+	{
+		return GetEditorCameraModifier();
+	}
+#endif
+	
+	return nullptr;
+}
+
+#if WITH_EDITOR
+UCameraModifier* UDaySequenceCameraModifierManager::GetEditorCameraModifier()
+{
+	if (const UWorld* World = GetWorld(); World && World->WorldType != EWorldType::Editor)
+	{
+		return nullptr;
+	}
+	
+	if (!EditorCameraModifier)
+	{
+		EditorCameraModifier = NewObject<UDaySequenceCameraModifier>(this, "EditorCameraModifier", RF_Transient);
+	}
+
+	if (!EditorCameraModiferPreview)
+	{
+		EditorCameraModiferPreview = NewObject<UPostProcessComponent>(GetOuter(), "EditorCameraModifierPreview", RF_Transient);
+		EditorCameraModiferPreview->RegisterComponent();
+		EditorCameraModiferPreview->AttachToComponent(EditorCameraModiferPreview->GetOwner()->GetRootComponent(), FAttachmentTransformRules::KeepRelativeTransform);
+	}
+
+	return EditorCameraModifier;
+}
+
+void UDaySequenceCameraModifierManager::UpdateEditorPreview() const
+{
+	if (EditorCameraModifier && EditorCameraModiferPreview)
+	{
+		EditorCameraModiferPreview->Settings = EditorCameraModifier->GetSettings();
+	}
+}
+#endif
\ No newline at end of file
diff --git a/Engine/Plugins/Experimental/DaySequence/Source/DaySequence/Private/DaySequenceModifierComponent.cpp b/Engine/Plugins/Experimental/DaySequence/Source/DaySequence/Private/DaySequenceModifierComponent.cpp
--- a/Engine/Plugins/Experimental/DaySequence/Source/DaySequence/Private/DaySequenceModifierComponent.cpp
+++ b/Engine/Plugins/Experimental/DaySequence/Source/DaySequence/Private/DaySequenceModifierComponent.cpp
@@ -3,15 +3,18 @@
 #include "DaySequenceModifierComponent.h"
 
 #include "DaySequence.h"
+#include "DaySequenceCameraModifier.h"
 #include "DaySequenceCollectionAsset.h"
 #include "DaySequenceModule.h"
 #include "DaySequencePlayer.h"
 #include "DaySequenceStaticTime.h"
+#include "DaySequenceSubsystem.h"
 #include "DaySequenceTrack.h"
 
 #include "MovieScene.h"
 #include "MovieSceneSection.h"
 #include "MovieSceneTrack.h"
+#include "Camera/CameraModifier.h"
 #include "Components/BoxComponent.h"
 #include "Components/CapsuleComponent.h"
 #include "Components/SphereComponent.h"
@@ -404,12 +407,21 @@ void UDaySequenceModifierComponent::RemoveSubSequenceTrack()
 		}
 	};
 	
-	for (const TWeakObjectPtr<UMovieSceneSubSection> SubSection : SubSections)
+	for (const TWeakObjectPtr<UMovieSceneSubSection> WeakSubSection : SubSections)
 	{
+		const UMovieSceneSubSection* SubSection = WeakSubSection.Get();
+		
 #if WITH_EDITOR
-		ADaySequenceActor::OnSubSectionRemovedEvent.Broadcast(SubSection.Get());
+		ADaySequenceActor::OnSubSectionRemovedEvent.Broadcast(SubSection);
 #endif
-		RemoveSubTrack(SubSection.Get());
+
+		// When we untrack a subsection, we need to remove any associated resolve functions.
+		if (TargetActor && SubSection)
+		{
+			TargetActor->UnregisterBindingResolveFunction(SubSection->GetSequenceID());
+		}
+		
+		RemoveSubTrack(SubSection);
 	}
 	SubSections.Empty();
 
@@ -759,6 +771,23 @@ UMovieSceneSubSection* UDaySequenceModifierComponent::InitializeDaySequence(cons
 		SubSections.Add(SubSection);
 	}
 
+	if (Entry.Sequence)
+	{
+		if (const FGuid CameraModifierBindingGuid = Entry.Sequence->GetSpecializedBinding(EDaySequenceBindingReferenceSpecialization::CameraModifier); CameraModifierBindingGuid.IsValid())
+		{
+			TargetActor->RegisterBindingResolveFunction(SubSection->GetSequenceID(), CameraModifierBindingGuid, [this](TArray<UObject*, TInlineAllocator<1>>& InOutObjects) -> bool
+			{
+				if (TargetActor)
+				{
+					UObject* CameraModifier = TargetActor->GetCameraModifierManager()->GetCameraModifier(WeakBlendTarget.Get());
+					InOutObjects.Add(CameraModifier);
+				}
+						
+				return false; // don't allow default lookup
+			});
+		}
+	}
+	
 	const TFunction<void(void)> SetSubTrackMuteStateConditional = [this, SubSection, Conditions = Entry.Conditions.Conditions]()
 	{
 		if (!IsValid(this) || !IsValid(SubSection))
@@ -894,6 +923,15 @@ EDayNightCycleMode UDaySequenceModifierComponent::GetDayNightCycle() const
 	return DayNightCycle;
 }
 
+void UDaySequenceModifierComponent::SetBias(int32 NewBias)
+{
+	Bias = NewBias;
+}
+int32 UDaySequenceModifierComponent::GetBias() const
+{
+	return Bias;
+}
+
 void UDaySequenceModifierComponent::SetDayNightCycleTime(float Time)
 {
 	DayNightCycleTime = Time;
diff --git a/Engine/Plugins/Experimental/DaySequence/Source/DaySequence/Public/Actors/DaySequenceModifierVolume.h b/Engine/Plugins/Experimental/DaySequence/Source/DaySequence/Public/Actors/DaySequenceModifierVolume.h
--- a/Engine/Plugins/Experimental/DaySequence/Source/DaySequence/Public/Actors/DaySequenceModifierVolume.h
+++ b/Engine/Plugins/Experimental/DaySequence/Source/DaySequence/Public/Actors/DaySequenceModifierVolume.h
@@ -27,54 +27,25 @@ public:
 protected:
 	//~ Begin AActor interface
 	virtual void BeginPlay() override;
+	virtual void EndPlay(const EEndPlayReason::Type EndPlayReason) override;
 	virtual void OnConstruction(const FTransform& Transform) override;
 	//~ End AActor interface
 	
-	// Primary initialization entry point to clarify initialization control flow and to consolidate editor and runtime initialization.
-	// Called from BeginPlay or OnConstruction
+	/** Primary initialization entry point to clarify initialization control flow and to consolidate editor and runtime initialization. */
 	void Initialize();
-
-	
-	
-	/** Player Controller Setup Functions */
-	
-	/**
-	 * At runtime this unconditionally calls CachePlayerController and SetupBlendTargetCallbacks.
-	 * In editor this conditionally calls CachePlayerController and SetupBlendTargetCallbacks if we are not in an Editor world.
-	 */
-	void PlayerControllerSetup();
-
-	/**
-	 * Attempts to find a local player controller.
-	 * If found, we get the controller's view target.
-	 * If not found, we call QueuePlayerControllerQuery which will result in another call to this function next tick.
-	 */
-	void CachePlayerController();
-
-	/** Set a timer to call CachePlayerController one tick from now. */
-	void QueuePlayerControllerQuery();
-	
-
+	void Deinitialize();
 	
-	/** Day Sequence Actor Setup Functions */
-
-	/** Dumb wrapper for SetupDaySequenceSubsystemCallbacks and BindToDaySequenceActor. */
+	/** Wrapper for SetupDaySequenceSubsystemCallbacks and BindToDaySequenceActor. */
 	void DaySequenceActorSetup();
 
 	/** Registers a callback that calls BindToDaySequenceActor when the world's current Day Sequence Actor changes. */
 	void SetupDaySequenceSubsystemCallbacks();
 
-	/**
-	 * Attempts to get the world's current Day Sequence Actor.
-	 * 
-	 * If a Day Sequence Actor is found that is not the currently bound actor,
-	 * we notify the day sequence modifier and call VolumeSetup and PlayerControllerSetup.
-	 */
+	/** Binds all modifier components to the world's current Day Sequence Actor. */
 	void BindToDaySequenceActor();
 	
-
-	
-	void SetBlendTarget(APlayerController* InPC);
+	/** Initializes a modifier component for this player controller. Will create an additional component if necessary. */
+	void CreatePlayer(APlayerController* InPC);
 	
 protected:
 	UPROPERTY(VisibleAnywhere, Category = "Day Sequence", BlueprintReadOnly, meta = (AllowPrivateAccess = "true"))
@@ -89,10 +60,17 @@ protected:
 	UPROPERTY(Transient, DuplicateTransient)
 	TObjectPtr<APlayerController> CachedPlayerController;
 	
-	UPROPERTY(Transient, DuplicateTransient)
-	TObjectPtr<AActor> CurrentBlendTarget;
-	
-	FDelegateHandle ViewTargetChangedHandle;
-	
 	FDelegateHandle ReplayScrubbedHandle;
+
+private:
+	bool IsSplitscreenSupported() const;
+	
+	UPROPERTY(EditAnywhere, AdvancedDisplay, Category = "Day Sequence", DisplayName = "Enable Experimental Splitscreen Support")
+	bool bEnableSplitscreenSupport;
+
+	/** Transient modifier components associated with additional local players (primarily for split screen support) */
+	UPROPERTY(Transient, DuplicateTransient)
+	TMap<TObjectPtr<APlayerController>, TObjectPtr<UDaySequenceModifierComponent>> AdditionalPlayers;
+
+	FDelegateHandle ActorSpawnedHandle;
 };
diff --git a/Engine/Plugins/Experimental/DaySequence/Source/DaySequence/Public/DaySequence.h b/Engine/Plugins/Experimental/DaySequence/Source/DaySequence/Public/DaySequence.h
--- a/Engine/Plugins/Experimental/DaySequence/Source/DaySequence/Public/DaySequence.h
+++ b/Engine/Plugins/Experimental/DaySequence/Source/DaySequence/Public/DaySequence.h
@@ -31,6 +31,9 @@ public:
 
 	void AddDefaultBinding(const FGuid& PossessableGuid);
 
+	void AddSpecializedBinding(EDaySequenceBindingReferenceSpecialization Specialization);
+	FGuid GetSpecializedBinding(EDaySequenceBindingReferenceSpecialization Specialization) const;
+
 	// UMovieSceneSequence interface
 	virtual void BindPossessableObject(const FGuid& ObjectId, UObject& PossessedObject, UObject* Context) override;
 	virtual bool CanPossessObject(UObject& Object, UObject* InPlaybackContext) const override;
diff --git a/Engine/Plugins/Experimental/DaySequence/Source/DaySequence/Public/DaySequenceActor.h b/Engine/Plugins/Experimental/DaySequence/Source/DaySequence/Public/DaySequenceActor.h
--- a/Engine/Plugins/Experimental/DaySequence/Source/DaySequence/Public/DaySequenceActor.h
+++ b/Engine/Plugins/Experimental/DaySequence/Source/DaySequence/Public/DaySequenceActor.h
@@ -29,6 +29,7 @@ namespace EEndPlayReason { enum Type : int; }
 
 class UCurveFloat;
 class UDaySequence;
+class UDaySequenceCameraModifierManager;
 class UDaySequenceCollectionAsset;
 class UDaySequenceConditionTag;
 class UDaySequencePlayer;
@@ -275,6 +276,11 @@ public:
 	void RegisterStaticTimeContributor(const UE::DaySequence::FStaticTimeContributor& NewContributor) const;
 	void UnregisterStaticTimeContributor(const UObject* InUserObject) const;
 	
+	void RegisterBindingResolveFunction(FMovieSceneSequenceID SequenceID, FGuid Guid, TFunction<bool(TArray<UObject*, TInlineAllocator<1>>&)> InFunction);
+	void UnregisterBindingResolveFunction(FMovieSceneSequenceID SequenceID, FGuid Guid = FGuid());
+
+	UDaySequenceCameraModifierManager* GetCameraModifierManager() { return CameraModifierManager; }
+	
 	UDaySequence* GetRootSequence() const;
 
 	virtual void UpdateRootSequence();
@@ -530,6 +536,9 @@ protected:
 	
 	TSharedPtr<UE::DaySequence::FStaticTimeManager> StaticTimeManager;
 
+	UPROPERTY(Transient)
+	TObjectPtr<UDaySequenceCameraModifierManager> CameraModifierManager;
+	
 #if WITH_EDITORONLY_DATA
 	/**
 	 * Blueprint exposed delegate invoked when the TimeOfDayPreview property
@@ -579,6 +588,8 @@ protected:
 
 	TArray<UMovieSceneSubSection*> SubSections;
 	
+	TMap<FMovieSceneSequenceID, TMap<FGuid, TFunction<bool(TArray<UObject*, TInlineAllocator<1>>&)>>> BindingResolveFunctions;
+	
 	UPROPERTY(Transient)
 	TMap<TSubclassOf<UDaySequenceConditionTag>, TObjectPtr<UDaySequenceConditionTag>> TrackConditionMap;
 };
diff --git a/Engine/Plugins/Experimental/DaySequence/Source/DaySequence/Public/DaySequenceBindingReference.h b/Engine/Plugins/Experimental/DaySequence/Source/DaySequence/Public/DaySequenceBindingReference.h
--- a/Engine/Plugins/Experimental/DaySequence/Source/DaySequence/Public/DaySequenceBindingReference.h
+++ b/Engine/Plugins/Experimental/DaySequence/Source/DaySequence/Public/DaySequenceBindingReference.h
@@ -5,6 +5,14 @@
 #include "Misc/Guid.h"
 #include "DaySequenceBindingReference.generated.h"
 
+UENUM()
+enum class EDaySequenceBindingReferenceSpecialization
+{
+	None,						// An unspecialized binding. Resolution is determined by ExternalObjectPath and ObjectPath given some context.
+	Root,						// An empty binding with this specialization will always resolve to the Root Day Sequence Actor.
+	CameraModifier,				// An empty binding with this specialization will attempt to resolve to a camera modifier associated with a modifier's blend target.
+};
+
 /**
  * An external reference to a DaySequence object, resolvable through an arbitrary context.
  * 
@@ -32,6 +40,11 @@ struct FDaySequenceBindingReference
 	 */
 	DAYSEQUENCE_API static FDaySequenceBindingReference DefaultRootBinding();
 
+	/**
+	 * Construct a new binding that always resolves based on specialization type. This is highly context specific.
+	 */
+	DAYSEQUENCE_API static FDaySequenceBindingReference SpecializedBinding(EDaySequenceBindingReferenceSpecialization Specialization);
+	
 	/**
 	 * Resolve this reference within the specified context
 	 *
@@ -45,6 +58,8 @@ struct FDaySequenceBindingReference
 	 */
 	DAYSEQUENCE_API bool operator==(const FDaySequenceBindingReference& Other) const;
 
+	DAYSEQUENCE_API EDaySequenceBindingReferenceSpecialization GetSpecialization() const { return Specialization; }
+	
 #if WITH_EDITORONLY_DATA
 
 	/**
@@ -63,6 +78,10 @@ private:
 	UPROPERTY()
 	FString ObjectPath;
 
+	/** Used when object path data is empty. Generally used for context-specific dynamic bindings. */
+	UPROPERTY()
+	EDaySequenceBindingReferenceSpecialization Specialization = EDaySequenceBindingReferenceSpecialization::None;
+	
 #if WITH_EDITORONLY_DATA
 	/** The class of the object path. */
 	UPROPERTY()
@@ -131,12 +150,30 @@ struct FDaySequenceBindingReferences
 	void AddBinding(const FGuid& ObjectId, UObject* InObject, UObject* InContext);
 
 	/**
-	 * Adds a default binding that always resolves to the day sequence actor
+	 * Adds a default binding that always resolves to the day sequence actor.
+	 * Prefer using AddSpecializedBinding with the Root specialization.
 	 *
 	 * @param ObjectId	The ID to associate the object with
 	 */
 	void AddDefaultBinding(const FGuid& ObjectId);
 
+	/**
+	 * Adds a specialized binding. 
+	 *
+	 * @param ObjectId	The ID to associate the object with
+	 * @param Specialization The type of specialized binding to create
+	 */
+	void AddSpecializedBinding(const FGuid& ObjectId, EDaySequenceBindingReferenceSpecialization Specialization);
+
+	/**
+	 * Finds a specialized binding. 
+	 *
+	 * @param Specialization The type of specialized binding to find.
+	 * @return Specialization The Guid of the found specialized binding, or a default constructed FGuid if not found.
+	 */
+	
+	FGuid FindSpecializedBinding(EDaySequenceBindingReferenceSpecialization Specialization) const;
+	
 	/**
 	 * Resolve a binding for the specified ID using a given context
 	 *
@@ -179,4 +216,9 @@ private:
 	/** A set of object binding IDs that relate to anim sequence instances (must be a child of USkeletalMeshComponent) */
 	UPROPERTY()
 	TSet<FGuid> AnimSequenceInstances;
+
+	UPROPERTY()
+	TMap<EDaySequenceBindingReferenceSpecialization, FGuid> SpecializedReferenceToGuid;
+	UPROPERTY()
+	TMap<FGuid, EDaySequenceBindingReferenceSpecialization> GuidToSpecializedReference;
 };
diff --git a/Engine/Plugins/Experimental/DaySequence/Source/DaySequence/Public/DaySequenceCameraModifier.h b/Engine/Plugins/Experimental/DaySequence/Source/DaySequence/Public/DaySequenceCameraModifier.h
new file mode 100644
--- /dev/null
+++ b/Engine/Plugins/Experimental/DaySequence/Source/DaySequence/Public/DaySequenceCameraModifier.h
@@ -0,0 +1,77 @@
+// Copyright Epic Games, Inc. All Rights Reserved.
+
+#pragma once
+
+#include "Containers/Map.h"
+#include "Camera/CameraModifier.h"
+
+#include "DaySequenceCameraModifier.generated.h"
+
+class APlayerController;
+
+class UPostProcessComponent;
+
+/**
+ * Provides:
+ * 1) an interface for getting UDaySequenceCameraModifiers associated with player controllers,
+ * 2) an editor only camera modifier for resolving camera modifier bindings in editor worlds,
+ * 3) an editor only post process component for visualizing the editor only camera modifier.
+ */
+UCLASS(MinimalAPI)
+class UDaySequenceCameraModifierManager : public UObject
+{
+	GENERATED_BODY()
+	
+public:
+	
+	UCameraModifier* GetCameraModifier(APlayerController* InPC);
+	
+#if WITH_EDITOR
+	UCameraModifier* GetEditorCameraModifier();
+	void UpdateEditorPreview() const;
+	void ResetEditorPreview() { EditorCameraModiferPreview = nullptr; }
+#endif
+
+private:
+	
+#if WITH_EDITORONLY_DATA
+	/** A camera modifier not associated with any player controller, used to resolve the camera modifier binding in editor. */
+	UPROPERTY(Transient, DuplicateTransient)
+	TObjectPtr<UDaySequenceCameraModifier> EditorCameraModifier;
+	
+	/** A post process component used to preview the animation of EditorCameraModifier. */
+    UPROPERTY(Transient, DuplicateTransient)
+    TObjectPtr<UPostProcessComponent> EditorCameraModiferPreview;
+#endif
+
+	// Note: We have no ownership over either the controller or the modifier. The modifiers are owned by the player camera manager.
+	TMap<TWeakObjectPtr<APlayerController>, TWeakObjectPtr<UCameraModifier>> CameraModifiers;
+
+};
+
+UCLASS(MinimalAPI)
+class UDaySequenceCameraModifier : public UCameraModifier
+{
+	GENERATED_BODY()
+
+public:
+
+	const FPostProcessSettings& GetSettings() const { return Settings; }
+
+private:
+	
+	/**
+	 * Post process settings to use for this modifier.
+	 * 
+	 * Sequencer should handle this (blending multiple sequences, restoring to default, etc),
+	 * we just handle forwarding it to the provided settings in ModifyPostProcess()
+	 */
+	UPROPERTY(EditAnywhere, interp, Category = PostProcess, meta = (ShowOnlyInnerProperties))
+	FPostProcessSettings Settings;
+	
+	virtual void ModifyPostProcess(float DeltaTime, float& PostProcessBlendWeight, FPostProcessSettings& PostProcessSettings) override
+	{
+		PostProcessBlendWeight = 1.f;
+		PostProcessSettings = Settings;
+	}
+};
diff --git a/Engine/Plugins/Experimental/DaySequence/Source/DaySequence/Public/DaySequenceModifierComponent.h b/Engine/Plugins/Experimental/DaySequence/Source/DaySequence/Public/DaySequenceModifierComponent.h
--- a/Engine/Plugins/Experimental/DaySequence/Source/DaySequence/Public/DaySequenceModifierComponent.h
+++ b/Engine/Plugins/Experimental/DaySequence/Source/DaySequence/Public/DaySequenceModifierComponent.h
@@ -176,6 +176,9 @@ public:
 	DAYSEQUENCE_API void SetDayNightCycle(EDayNightCycleMode NewMode);
 	DAYSEQUENCE_API EDayNightCycleMode GetDayNightCycle() const;
 
+	DAYSEQUENCE_API void SetBias(int32 NewBias);
+	DAYSEQUENCE_API int32 GetBias() const;
+	
 	DAYSEQUENCE_API void SetDayNightCycleTime(float Time);
 	DAYSEQUENCE_API float GetDayNightCycleTime() const;
 	
@@ -251,7 +254,7 @@ protected:
 	TArray<FComponentReference> VolumeShapeComponents;
 
 	/** The actor to use for distance-based volume blend calculations */
-	UPROPERTY(Transient)
+	UPROPERTY(Transient, DuplicateTransient)
 	TWeakObjectPtr<APlayerController> WeakBlendTarget;
 
 	/** The user provided Day Sequence. */
@@ -270,7 +273,7 @@ protected:
 	TObjectPtr<UDaySequenceModifierEasingFunction> EasingFunction;
 	
 	/** User-defined bias. */
-	UPROPERTY(BlueprintReadWrite, EditAnywhere, Category="Day Sequence", meta=(EditCondition="!bIgnoreBias"))
+	UPROPERTY(EditAnywhere, BlueprintReadWrite, Setter, Getter, Category="Day Sequence", meta=(EditCondition="!bIgnoreBias"))
 	int32 Bias;
 
 	/** The time to use for the day/night cycle. */
diff --git a/Engine/Plugins/Experimental/DaySequence/Source/DaySequenceEditor/Private/DaySequenceEditorModule.cpp b/Engine/Plugins/Experimental/DaySequence/Source/DaySequenceEditor/Private/DaySequenceEditorModule.cpp
--- a/Engine/Plugins/Experimental/DaySequence/Source/DaySequenceEditor/Private/DaySequenceEditorModule.cpp
+++ b/Engine/Plugins/Experimental/DaySequence/Source/DaySequenceEditor/Private/DaySequenceEditorModule.cpp
@@ -9,6 +9,7 @@
 #include "DaySequenceActorPreview.h"
 #include "DaySequenceEditorActorBinding.h"
 #include "DaySequenceEditorActorSpawner.h"
+#include "DaySequenceEditorSpecializedBinding.h"
 #include "DaySequenceSubsystem.h"
 #include "DaySequenceConditionSetCustomization.h"
 #include "DaySequenceTimeDetailsCustomization.h"
@@ -91,6 +92,7 @@ public:
 	void OnOpenDaySequenceActor() const;
 	bool CanOpenDaySequenceActor() const;
 	static TSharedRef<ISequencerEditorObjectBinding> OnCreateActorBinding(TSharedRef<ISequencer> InSequencer);
+	static TSharedRef<ISequencerEditorObjectBinding> OnCreateSpecializedBinding(TSharedRef<ISequencer> InSequencer);
 	
 	TSharedPtr<FPropertySection> RegisterPropertySection(FPropertyEditorModule& PropertyModule, FName ClassName, FName SectionName, FText DisplayName);
 	void RegisterModulePropertySections();
@@ -134,6 +136,8 @@ private:
 
 	FDelegateHandle ActorBindingDelegateHandle;
 
+	FDelegateHandle SpecializedBindingDelegateHandle;
+
 	FDelegateHandle EditorActorSpawnerDelegateHandle;
 
 	FDelegateHandle SequenceEditorHandle;
@@ -725,6 +729,8 @@ void FDaySequenceEditorModule::RegisterEditorObjectBindings()
 {
 	ISequencerModule& SequencerModule = FModuleManager::LoadModuleChecked<ISequencerModule>("Sequencer");
 	ActorBindingDelegateHandle = SequencerModule.RegisterEditorObjectBinding(FOnCreateEditorObjectBinding::CreateStatic(&FDaySequenceEditorModule::OnCreateActorBinding));
+	SpecializedBindingDelegateHandle = SequencerModule.RegisterEditorObjectBinding(FOnCreateEditorObjectBinding::CreateStatic(&FDaySequenceEditorModule::OnCreateSpecializedBinding));
+
 }
 
 /** Unregisters sequencer editor object bindings */
@@ -734,6 +740,7 @@ void FDaySequenceEditorModule::DeregisterEditorObjectBindings()
 	if (SequencerModule)
 	{
 		SequencerModule->UnRegisterEditorObjectBinding(ActorBindingDelegateHandle);
+		SequencerModule->UnRegisterEditorObjectBinding(SpecializedBindingDelegateHandle);
 	}
 }
 
@@ -758,6 +765,11 @@ TSharedRef<ISequencerEditorObjectBinding> FDaySequenceEditorModule::OnCreateActo
 	return MakeShareable(new FDaySequenceEditorActorBinding(InSequencer));
 }
 
+TSharedRef<ISequencerEditorObjectBinding> FDaySequenceEditorModule::OnCreateSpecializedBinding(TSharedRef<ISequencer> InSequencer)
+{
+	return MakeShareable(new FDaySequenceEditorSpecializedBinding(InSequencer));
+}
+
 void FDaySequenceEditorModule::RegisterSequenceEditor()
 {
 	ISequencerModule& SequencerModule = FModuleManager::LoadModuleChecked<ISequencerModule>("Sequencer");
diff --git a/Engine/Plugins/Experimental/DaySequence/Source/DaySequenceEditor/Private/DaySequenceEditorSpecializedBinding.cpp b/Engine/Plugins/Experimental/DaySequence/Source/DaySequenceEditor/Private/DaySequenceEditorSpecializedBinding.cpp
new file mode 100644
--- /dev/null
+++ b/Engine/Plugins/Experimental/DaySequence/Source/DaySequenceEditor/Private/DaySequenceEditorSpecializedBinding.cpp
@@ -0,0 +1,81 @@
+// Copyright Epic Games, Inc. All Rights Reserved.
+
+#include "DaySequenceEditorSpecializedBinding.h"
+
+#include "DaySequence.h"
+#include "ISequencer.h"
+#include "Framework/MultiBox/MultiBoxBuilder.h"
+
+#define LOCTEXT_NAMESPACE "DaySequenceEditorSpecializedBinding"
+
+FDaySequenceEditorSpecializedBinding::FDaySequenceEditorSpecializedBinding(TSharedRef<ISequencer> InSequencer)
+	: Sequencer(InSequencer)
+{
+}
+
+void FDaySequenceEditorSpecializedBinding::BuildSequencerAddMenu(FMenuBuilder& MenuBuilder)
+{
+	// todo [nickolas.drake]: add icon
+	MenuBuilder.AddSubMenu(
+		LOCTEXT("SpecializedBindingLabel", "Specialized Bindings"),
+		LOCTEXT("SpecializedBindingToolTip", "Add specialized binding types which have special resolution rules"),
+		FNewMenuDelegate::CreateRaw(this, &FDaySequenceEditorSpecializedBinding::AddSpecializedBindingMenuExtensions)
+	);
+}
+
+bool FDaySequenceEditorSpecializedBinding::SupportsSequence(UMovieSceneSequence* InSequence) const
+{
+	return InSequence->GetClass() == UDaySequence::StaticClass();
+}
+
+void FDaySequenceEditorSpecializedBinding::AddSpecializedBindingMenuExtensions(FMenuBuilder& MenuBuilder)
+{
+	// generates an FUIAction for a binding specialization type
+	auto GetAddBindingAction = [WeakSequencer = Sequencer](EDaySequenceBindingReferenceSpecialization Specialization)
+	{
+		FUIAction Action;
+	
+		Action.ExecuteAction = FExecuteAction::CreateLambda([WeakSequencer, Specialization]()
+		{
+			if (const TSharedPtr<ISequencer> Sequencer = WeakSequencer.Pin())
+			{
+				if (UDaySequence* FocusedSequence = Cast<UDaySequence>(Sequencer->GetFocusedMovieSceneSequence()))
+				{
+					return FocusedSequence->AddSpecializedBinding(Specialization);
+				}
+			}
+		});
+
+		Action.CanExecuteAction = FCanExecuteAction::CreateLambda([WeakSequencer, Specialization]()
+		{
+			if (const TSharedPtr<ISequencer> Sequencer =  WeakSequencer.Pin())
+			{
+				if (const UDaySequence* FocusedSequence = Cast<UDaySequence>(Sequencer->GetFocusedMovieSceneSequence()))
+				{
+					// only allow creation if binding type isn't present already
+					return !FocusedSequence->GetSpecializedBinding(Specialization).IsValid();
+				}
+			}
+		
+			return false;
+		});
+
+		return Action;
+	};
+
+	// Add an entry for a root binding
+	MenuBuilder.AddMenuEntry(LOCTEXT("RootActorBindingLabel", "Root Day Sequence Actor Binding"),
+		LOCTEXT("RootActorBindingTooltip", "Add a new root actor binding. This allows a sequence to animate a generic day sequence actor."),
+		FSlateIcon(),
+		GetAddBindingAction(EDaySequenceBindingReferenceSpecialization::Root)
+	);
+	
+	// Add an entry for a camera modifier binding
+	MenuBuilder.AddMenuEntry(LOCTEXT("CameraModifierBindingLabel", "Camera Modifier Binding"),
+		LOCTEXT("CameraModifierBindingTooltip", "Add a new camera modifier binding. This allows a sequence to animate post process effects for local players."),
+		FSlateIcon(),
+		GetAddBindingAction(EDaySequenceBindingReferenceSpecialization::CameraModifier)
+	);
+}
+
+#undef LOCTEXT_NAMESPACE
diff --git a/Engine/Plugins/Experimental/DaySequence/Source/DaySequenceEditor/Private/DaySequenceEditorSpecializedBinding.h b/Engine/Plugins/Experimental/DaySequence/Source/DaySequenceEditor/Private/DaySequenceEditorSpecializedBinding.h
new file mode 100644
--- /dev/null
+++ b/Engine/Plugins/Experimental/DaySequence/Source/DaySequenceEditor/Private/DaySequenceEditorSpecializedBinding.h
@@ -0,0 +1,28 @@
+// Copyright Epic Games, Inc. All Rights Reserved.
+
+#pragma once
+
+#include "ISequencerEditorObjectBinding.h"
+#include "Templates/SharedPointer.h"
+
+class ISequencer;
+class AActor;
+
+class FDaySequenceEditorSpecializedBinding : public ISequencerEditorObjectBinding
+{
+public:
+
+	FDaySequenceEditorSpecializedBinding(TSharedRef<ISequencer> InSequencer);
+
+	// ISequencerEditorObjectBinding interface
+	virtual void BuildSequencerAddMenu(FMenuBuilder& MenuBuilder) override;
+	virtual bool SupportsSequence(UMovieSceneSequence* InSequence) const override;
+
+private:
+
+	/** Menu extension callback for the specialized binding menu */
+	void AddSpecializedBindingMenuExtensions(FMenuBuilder& MenuBuilder);
+
+private:
+	TWeakPtr<ISequencer> Sequencer;
+};
