Fix Mac compile error

diff --git a/Engine/Plugins/Experimental/DaySequence/Source/DaySequenceEditor/Private/DaySequenceEditorSpecializedBinding.cpp b/Engine/Plugins/Experimental/DaySequence/Source/DaySequenceEditor/Private/DaySequenceEditorSpecializedBinding.cpp
--- a/Engine/Plugins/Experimental/DaySequence/Source/DaySequenceEditor/Private/DaySequenceEditorSpecializedBinding.cpp
+++ b/Engine/Plugins/Experimental/DaySequence/Source/DaySequenceEditor/Private/DaySequenceEditorSpecializedBinding.cpp
@@ -37,9 +37,9 @@ void FDaySequenceEditorSpecializedBinding::AddSpecializedBindingMenuExtensions(F
 	
 		Action.ExecuteAction = FExecuteAction::CreateLambda([WeakSequencer, Specialization]()
 		{
-			if (const TSharedPtr<ISequencer> Sequencer = WeakSequencer.Pin())
+			if (const TSharedPtr<ISequencer> PinnedSequencer = WeakSequencer.Pin())
 			{
-				if (UDaySequence* FocusedSequence = Cast<UDaySequence>(Sequencer->GetFocusedMovieSceneSequence()))
+				if (UDaySequence* FocusedSequence = Cast<UDaySequence>(PinnedSequencer->GetFocusedMovieSceneSequence()))
 				{
 					return FocusedSequence->AddSpecializedBinding(Specialization);
 				}
@@ -48,9 +48,9 @@ void FDaySequenceEditorSpecializedBinding::AddSpecializedBindingMenuExtensions(F
 
 		Action.CanExecuteAction = FCanExecuteAction::CreateLambda([WeakSequencer, Specialization]()
 		{
-			if (const TSharedPtr<ISequencer> Sequencer =  WeakSequencer.Pin())
+			if (const TSharedPtr<ISequencer> PinnedSequencer = WeakSequencer.Pin())
 			{
-				if (const UDaySequence* FocusedSequence = Cast<UDaySequence>(Sequencer->GetFocusedMovieSceneSequence()))
+				if (const UDaySequence* FocusedSequence = Cast<UDaySequence>(PinnedSequencer->GetFocusedMovieSceneSequence()))
 				{
 					// only allow creation if binding type isn't present already
 					return !FocusedSequence->GetSpecializedBinding(Specialization).IsValid();
