DaySequenceModifierComponent: Fix 2 issues.

1: This change corrects an issue where the modifier would never stop overriding the DSA's update interval once started.
2: This change makes the modifier respect visibility. This is so that data layers with modifiers can remain loaded but hidden and correctly do nothing.

diff --git a/Engine/Plugins/Experimental/DaySequence/Source/DaySequence/Private/DaySequenceModifierComponent.cpp b/Engine/Plugins/Experimental/DaySequence/Source/DaySequence/Private/DaySequenceModifierComponent.cpp
--- a/Engine/Plugins/Experimental/DaySequence/Source/DaySequence/Private/DaySequenceModifierComponent.cpp
+++ b/Engine/Plugins/Experimental/DaySequence/Source/DaySequence/Private/DaySequenceModifierComponent.cpp
@@ -88,6 +88,14 @@ namespace UE::DaySequence
 		// @todo: unsupported shape?
 		return (InWorldPosition - ShapeComponent->GetComponentLocation()).Length();
 	}
+	
+	static TAutoConsoleVariable<bool> CVarModifierDisableWhenInvisible(
+		TEXT("DaySequence.Modifier.DisableWhenInvisible"),
+		true,
+		TEXT("When true, day sequence modifier components will automatically disable when they are made invisible."),
+		ECVF_Default
+	);
+	
 } // namespace UE::DaySequence
 
 void UDaySequenceModifierEasingFunction::Initialize(EEasingFunctionType EasingType)
@@ -263,6 +271,8 @@ void UDaySequenceModifierComponent::OnUnregister()
 void UDaySequenceModifierComponent::DaySequenceUpdate()
 {
 	CSV_SCOPED_TIMING_STAT(DaySequence, SequencePlayerUpdated);
+
+	bool bWantsOverride = false;
 	
 	if (bIsComponentEnabled)
 	{
@@ -280,21 +290,26 @@ void UDaySequenceModifierComponent::DaySequenceUpdate()
 			DisableModifier();
 		}
 
-		if (OverrideUpdateIntervalHandle)
+		if (bIsEnabled && bSmoothBlending)
 		{
-			if (bIsEnabled && bSmoothBlending)
+			if (bForceSmoothBlending || FMath::IsWithin(GetBlendWeight(), 0.f + UE_KINDA_SMALL_NUMBER, 1.f - UE_KINDA_SMALL_NUMBER))
 			{
-				if (bForceSmoothBlending || FMath::IsWithin(GetBlendWeight(), 0.f + UE_KINDA_SMALL_NUMBER, 1.f - UE_KINDA_SMALL_NUMBER))
-				{
-					OverrideUpdateIntervalHandle->StartOverriding();
-				}
-				else
-				{
-					OverrideUpdateIntervalHandle->StopOverriding();
-				}
+				bWantsOverride = true;
 			}
 		}
 	}
+	
+	if (OverrideUpdateIntervalHandle)
+	{
+		if (bWantsOverride)
+		{
+			OverrideUpdateIntervalHandle->StartOverriding();
+		}
+		else
+		{
+			OverrideUpdateIntervalHandle->StopOverriding();
+		}
+	}
 }
 
 void UDaySequenceModifierComponent::BeginPlay()
@@ -1103,6 +1118,14 @@ float UDaySequenceModifierComponent::UpdateInternalBlendWeight()
 		break;
 	}
 
+	if (UE::DaySequence::CVarModifierDisableWhenInvisible.GetValueOnAnyThread())
+	{
+		if (!IsVisible())
+		{
+			InternalBlendWeight = 0.f;
+		}
+	}
+
 	return InternalBlendWeight;
 }
 
