DaySequence: Restore default engine material to the volumetric cloud component on the BaseDaySequenceActor

diff --git a/Engine/Plugins/Experimental/DaySequence/Source/DaySequence/Private/Actors/BaseDaySequenceActor.cpp b/Engine/Plugins/Experimental/DaySequence/Source/DaySequence/Private/Actors/BaseDaySequenceActor.cpp
--- a/Engine/Plugins/Experimental/DaySequence/Source/DaySequence/Private/Actors/BaseDaySequenceActor.cpp
+++ b/Engine/Plugins/Experimental/DaySequence/Source/DaySequence/Private/Actors/BaseDaySequenceActor.cpp
@@ -40,6 +40,13 @@ ABaseDaySequenceActor::ABaseDaySequenceActor(const FObjectInitializer& Init)
 	VolumetricCloudComponent = CreateOptionalDefaultSubobject<UVolumetricCloudComponent>(TEXT("VolumetricCloud"));
     VolumetricCloudComponent->SetupAttachment(RootComponent);
 
+	// We don't want to load this material for the CDO as it will hold on to it forever and it is quite a large asset.
+	if (!IsTemplate())
+	{
+		static ConstructorHelpers::FObjectFinder<UMaterialInterface> VolumetricCloudDefaultMaterialRef(TEXT("/Engine/EngineSky/VolumetricClouds/m_SimpleVolumetricCloud_Inst.m_SimpleVolumetricCloud_Inst"));
+		VolumetricCloudComponent->SetMaterial(VolumetricCloudDefaultMaterialRef.Object);
+	}
+
 	static ConstructorHelpers::FObjectFinder<UStaticMesh> SkySphereDefaultMesh(TEXT("/Engine/EngineSky/SM_SkySphere.SM_SkySphere"));
 	static ConstructorHelpers::FObjectFinder<UMaterialInterface> SkySphereDefaultMaterial(TEXT("/Engine/EngineSky/M_SimpleSkyDome.M_SimpleSkyDome"));
 	SkySphereComponent = CreateOptionalDefaultSubobject<UStaticMeshComponent>(TEXT("SkySphere"));
