DaySequenceModifierComponent: Set DaySequence.Modifier.DisableWhenInvisible to false by default to fix a regression caused by its introduction.

diff --git a/Engine/Plugins/Experimental/DaySequence/Source/DaySequence/Private/DaySequenceModifierComponent.cpp b/Engine/Plugins/Experimental/DaySequence/Source/DaySequence/Private/DaySequenceModifierComponent.cpp
--- a/Engine/Plugins/Experimental/DaySequence/Source/DaySequence/Private/DaySequenceModifierComponent.cpp
+++ b/Engine/Plugins/Experimental/DaySequence/Source/DaySequence/Private/DaySequenceModifierComponent.cpp
@@ -91,7 +91,7 @@ namespace UE::DaySequence
 	
 	static TAutoConsoleVariable<bool> CVarModifierDisableWhenInvisible(
 		TEXT("DaySequence.Modifier.DisableWhenInvisible"),
-		true,
+		false,
 		TEXT("When true, day sequence modifier components will automatically disable when they are made invisible."),
 		ECVF_Default
 	);
