DaySequence: Support native condition tags

diff --git a/Engine/Plugins/Experimental/DaySequence/Source/DaySequence/Private/DaySequenceConditionTag.cpp b/Engine/Plugins/Experimental/DaySequence/Source/DaySequence/Private/DaySequenceConditionTag.cpp
--- a/Engine/Plugins/Experimental/DaySequence/Source/DaySequence/Private/DaySequenceConditionTag.cpp
+++ b/Engine/Plugins/Experimental/DaySequence/Source/DaySequence/Private/DaySequenceConditionTag.cpp
@@ -31,7 +31,7 @@ void UDaySequenceConditionTag::SetupOnConditionValueChanged_Implementation() con
 {
 }
 
-void UDaySequenceConditionTag::BroadcastOnConditionValueChanged()
+void UDaySequenceConditionTag::BroadcastOnConditionValueChanged() const
 {
 	const bool bResult = Evaluate();
 	if (!bCachedEvalResult.IsSet() || bCachedEvalResult.GetValue() != bResult)
diff --git a/Engine/Plugins/Experimental/DaySequence/Source/DaySequence/Public/DaySequenceConditionTag.h b/Engine/Plugins/Experimental/DaySequence/Source/DaySequence/Public/DaySequenceConditionTag.h
--- a/Engine/Plugins/Experimental/DaySequence/Source/DaySequence/Public/DaySequenceConditionTag.h
+++ b/Engine/Plugins/Experimental/DaySequence/Source/DaySequence/Public/DaySequenceConditionTag.h
@@ -48,7 +48,7 @@ protected:
 	 * This will only trigger a broadcast if Evaluate() returns a different value than the last invocation of this function.
 	 */
 	UFUNCTION(BlueprintCallable, Category = "General")
-	void BroadcastOnConditionValueChanged();
+	void BroadcastOnConditionValueChanged() const;
 
 protected:
 	/**
@@ -62,5 +62,5 @@ private:
 	FOnConditionValueChanged OnConditionValueChanged;
 
 	/* This is an optional because it is unset until the first time BroadcastOnConditionValueChanged is called. */
-	TOptional<bool> bCachedEvalResult;
+	mutable TOptional<bool> bCachedEvalResult;
 };
diff --git a/Engine/Plugins/Experimental/DaySequence/Source/DaySequenceEditor/Private/SDaySequenceConditionSetPicker.cpp b/Engine/Plugins/Experimental/DaySequence/Source/DaySequenceEditor/Private/SDaySequenceConditionSetPicker.cpp
--- a/Engine/Plugins/Experimental/DaySequence/Source/DaySequenceEditor/Private/SDaySequenceConditionSetPicker.cpp
+++ b/Engine/Plugins/Experimental/DaySequence/Source/DaySequenceEditor/Private/SDaySequenceConditionSetPicker.cpp
@@ -24,6 +24,7 @@
 #include "Engine/BlueprintGeneratedClass.h"
 #include "Widgets/Layout/SScrollBar.h"
 #include "Widgets/Views/SListView.h"
+#include "UObject/UObjectIterator.h"
 
 #define LOCTEXT_NAMESPACE "DaySequenceConditionSetPicker"
 
@@ -215,6 +216,26 @@ void SDaySequenceConditionSetPicker::PopulateVisibleClasses()
 			return true;	// Returning false will halt the enumeration
 		});
 	}
+
+	// Scan for native classes
+	{
+		const UClass* ParentClass = UDaySequenceConditionTag::StaticClass();
+
+		auto IsNativeSubclass = [ParentClass](const UClass* InClass)
+		{
+			return InClass->IsChildOf(ParentClass)
+				&& InClass != ParentClass
+				&& (InClass->ClassFlags & CLASS_Native) != CLASS_None;
+		};
+		
+		for (TObjectIterator<UClass> It; It; ++It)
+		{
+			if (UClass* CurrentClass = *It; CurrentClass && IsNativeSubclass(CurrentClass))
+			{
+				AllConditionTags.Add(CurrentClass);
+			}
+		}
+	}
 }
 
 void SDaySequenceConditionSetPicker::PopulateCheckedTags()
