DaySequence: Allow procedural sequences to use condition tags.

diff --git a/Engine/Plugins/Experimental/DaySequence/Source/DaySequence/Private/DaySequenceActor.cpp b/Engine/Plugins/Experimental/DaySequence/Source/DaySequence/Private/DaySequenceActor.cpp
--- a/Engine/Plugins/Experimental/DaySequence/Source/DaySequence/Private/DaySequenceActor.cpp
+++ b/Engine/Plugins/Experimental/DaySequence/Source/DaySequence/Private/DaySequenceActor.cpp
@@ -736,7 +736,15 @@ void ADaySequenceActor::InitializeDaySequences()
 				continue;
 			}
 			
-			InitializeDaySequence(ProceduralDaySequence.GetMutable<FProceduralDaySequence>().GetSequence(this));
+			FProceduralDaySequence& ProceduralSequence = ProceduralDaySequence.GetMutable<FProceduralDaySequence>();
+ 
+			if (UDaySequence* Sequence = ProceduralSequence.GetSequence(this))
+			{
+				FDaySequenceCollectionEntry TempEntry(Sequence);
+				TempEntry.Conditions = ProceduralSequence.Conditions;
+  
+				InitializeDaySequence(TempEntry);
+			}
 		}
 	}
 }
diff --git a/Engine/Plugins/Experimental/DaySequence/Source/DaySequence/Private/DaySequenceModifierComponent.cpp b/Engine/Plugins/Experimental/DaySequence/Source/DaySequence/Private/DaySequenceModifierComponent.cpp
--- a/Engine/Plugins/Experimental/DaySequence/Source/DaySequence/Private/DaySequenceModifierComponent.cpp
+++ b/Engine/Plugins/Experimental/DaySequence/Source/DaySequence/Private/DaySequenceModifierComponent.cpp
@@ -682,7 +682,15 @@ void UDaySequenceModifierComponent::ReinitializeSubSequence(ADaySequenceActor::F
 						continue;
 					}
 			
-					InitializeDaySequence(ProceduralEntry.GetMutable<FProceduralDaySequence>().GetSequence(TargetActor));
+					FProceduralDaySequence& ProceduralSequence = ProceduralEntry.GetMutable<FProceduralDaySequence>();
+ 
+					if (UDaySequence* Sequence = ProceduralSequence.GetSequence(TargetActor))
+					{
+						FDaySequenceCollectionEntry TempEntry(Sequence);
+						TempEntry.Conditions = ProceduralSequence.Conditions;
+  
+						InitializeDaySequence(TempEntry);
+					}
 				}
 			}
 		}
diff --git a/Engine/Plugins/Experimental/DaySequence/Source/DaySequence/Public/ProceduralDaySequence.h b/Engine/Plugins/Experimental/DaySequence/Source/DaySequence/Public/ProceduralDaySequence.h
--- a/Engine/Plugins/Experimental/DaySequence/Source/DaySequence/Public/ProceduralDaySequence.h
+++ b/Engine/Plugins/Experimental/DaySequence/Source/DaySequence/Public/ProceduralDaySequence.h
@@ -2,6 +2,7 @@
 
 #pragma once
 
+#include "DaySequenceConditionSet.h"
 #include "ProceduralDaySequenceBuilder.h"
 
 #include "GameFramework/Actor.h"
@@ -44,6 +45,9 @@ struct DAYSEQUENCE_API FProceduralDaySequence
 	
 	UDaySequence* GetSequence(ADaySequenceActor* InActor);
 
+	UPROPERTY(EditAnywhere, Category="Day Sequence")
+	FDaySequenceConditionSet Conditions;
+	
 protected:
 	
 	virtual void BuildSequence(UProceduralDaySequenceBuilder* InBuilder) {}
