DaySequence: Dynamically decrease TickResolution if TimePerCycle is chosen such that the playback range of the root sequence would overflow.

diff --git a/Engine/Plugins/Experimental/DaySequence/Source/DaySequence/Private/DaySequenceActor.cpp b/Engine/Plugins/Experimental/DaySequence/Source/DaySequence/Private/DaySequenceActor.cpp
--- a/Engine/Plugins/Experimental/DaySequence/Source/DaySequence/Private/DaySequenceActor.cpp
+++ b/Engine/Plugins/Experimental/DaySequence/Source/DaySequence/Private/DaySequenceActor.cpp
@@ -73,8 +73,51 @@ namespace UE::DaySequence
 	false,
 	TEXT("When true, calls to SetTimeOfDay will dump the stack trace."),
 	ECVF_Default
-);
-	
+	);
+
+	bool IsTimePerCycleOverflow(const FDaySequenceTime& InTimePerCycle, const UMovieScene& InMovieScene)
+	{
+		const FFrameRate DisplayRate = InMovieScene.GetDisplayRate();
+		const FFrameRate TickResolution = InMovieScene.GetTickResolution();
+		const float TimePerCycleSeconds = InTimePerCycle.ToSeconds();
+		const int32 CycleSeconds = static_cast<int32>(TimePerCycleSeconds) + 1;
+		const int32 DisplayFactor = DisplayRate.Numerator / DisplayRate.Denominator;
+		const int32 TickFactor = TickResolution.Numerator / TickResolution.Denominator;
+		const int32 NewCycleSeconds = CycleSeconds * DisplayFactor * TickFactor;
+		const int32 CheckCycleSeconds = NewCycleSeconds / (DisplayFactor * TickFactor);
+		return (CheckCycleSeconds != CycleSeconds);
+	}
+
+	FFrameRate GetMaximumTickResolution(const FDaySequenceTime& InTimePerCycle, const UMovieScene& InMovieScene)
+	{
+		const FFrameRate DisplayRate = InMovieScene.GetDisplayRate();
+		const int32 DisplayFactor = DisplayRate.Numerator / DisplayRate.Denominator;
+		int32 MaxResolution = std::numeric_limits<int32>::max() / (InTimePerCycle.ToSeconds() * DisplayFactor);
+
+		// Set MaxResolution to the closest multiple of DisplayFactor that is less than the initially computed value of MaxResolution.
+		MaxResolution = (MaxResolution / DisplayFactor) * DisplayFactor;
+		
+		return FFrameRate(MaxResolution, 1);
+	}
+
+	void SetRootSequencePlaybackRange(const FDaySequenceTime& InTimePerCycle, const UMovieSceneSequence& InRootSequence)
+	{
+		if (UMovieScene* RootMovieScene = InRootSequence.GetMovieScene())
+		{
+			// Drop tick resolution if we are going to overflow playback range at this TPC
+			if (IsTimePerCycleOverflow(InTimePerCycle, *RootMovieScene))
+			{
+				RootMovieScene->SetTickResolutionDirectly(GetMaximumTickResolution(InTimePerCycle, *RootMovieScene));
+			}
+			
+			const int32 RootDuration = RootMovieScene->GetTickResolution().AsFrameNumber(InTimePerCycle.ToSeconds()).Value;
+			RootMovieScene->SetPlaybackRange(0, RootDuration);
+#if WITH_EDITOR
+			RootMovieScene->SetPlaybackRangeLocked(true);
+#endif
+			RootMovieScene->MarkAsChanged();
+		}
+	}
 } // namespace UE::DaySequence
 
 ADaySequenceActor::ADaySequenceActor(const FObjectInitializer& Init)
@@ -541,6 +584,8 @@ void ADaySequenceActor::InitializePlayer()
 
 void ADaySequenceActor::InitializeRootSequence()
 {
+	using namespace UE::DaySequence;
+	
 	CSV_SCOPED_TIMING_STAT(DaySequence, InitializeRootSequence);
 
 	if (IsTemplate())
@@ -571,13 +616,7 @@ void ADaySequenceActor::InitializeRootSequence()
 		RootSequence->Initialize(RF_Transient);
 		RootSequence->SetSequenceFlags(EMovieSceneSequenceFlags::Volatile);
 
-		UMovieScene* RootMovieScene = RootSequence->GetMovieScene();
-		const float DaySeconds = TimePerCycle.ToSeconds();
-		const int32 RootDuration = RootMovieScene->GetTickResolution().AsFrameNumber(DaySeconds).Value;
-		RootMovieScene->SetPlaybackRange(0, RootDuration);
-#if WITH_EDITOR
-		RootMovieScene->SetPlaybackRangeLocked(true);
-#endif
+		SetRootSequencePlaybackRange(TimePerCycle, *RootSequence);
 #if ROOT_SEQUENCE_RECONSTRUCTION_ENABLED
 	}
 #endif
@@ -1137,6 +1176,8 @@ void ADaySequenceActor::SetTimePerCycle(float InHours)
 
 void ADaySequenceActor::Multicast_SetTimePerCycle_Implementation(float InHours)
 {
+	using namespace UE::DaySequence;
+	
 	const TOptional<FFrameRate> FrameRate = RootSequence ? RootSequence->GetMovieScene()->GetDisplayRate() : TOptional<FFrameRate>();
 
 	const bool bIsGameWorld = GetWorld()->IsGameWorld();
@@ -1147,47 +1188,29 @@ void ADaySequenceActor::Multicast_SetTimePerCycle_Implementation(float InHours)
 	const FDaySequenceTime NewTimePerCycle = FDaySequenceTime::FromHours(InHours);
 	if (NewTimePerCycle != TimePerCycle)
 	{
-		// Validate the new TimePerCycle to avoid overflowing the playback range.
 		UMovieScene* RootMovieScene = RootSequence->GetMovieScene();
-		const FFrameRate DisplayRate = RootMovieScene->GetDisplayRate();
-		const FFrameRate TickResolution = RootMovieScene->GetTickResolution();
-		const auto IsTimePerCycleOverflow = [&DisplayRate, &TickResolution](const FDaySequenceTime& InTimePerCycle)
-		{
-			const float TimePerCycleSeconds = InTimePerCycle.ToSeconds();
-			const int32 CycleSeconds = static_cast<int32>(TimePerCycleSeconds) + 1;
-			const int32 DisplayFactor = DisplayRate.Numerator / DisplayRate.Denominator;
-			const int32 TickFactor = TickResolution.Numerator / TickResolution.Denominator;
-			const int32 NewCycleSeconds = CycleSeconds * DisplayFactor * TickFactor;
-			const int32 CheckCycleSeconds = NewCycleSeconds / (DisplayFactor * TickFactor);
-			return (CheckCycleSeconds != CycleSeconds);
-		};
-		if (IsTimePerCycleOverflow(NewTimePerCycle))
-		{
-			UE_LOG(LogDaySequence, Warning, TEXT("Skipping SetTimePerCycle( %f hours ) to avoid overflowing playback range"), InHours);
-			return;
-		}
 		
 		TimePerCycle = NewTimePerCycle;
 		
-		if (bIsGameWorld)
-		{
-			// Update playback range for the root sequence.
-			const float DaySeconds = TimePerCycle.ToSeconds();
-			const int32 RootDuration = RootMovieScene->GetTickResolution().AsFrameNumber(DaySeconds).Value;
-			RootMovieScene->MarkAsChanged();
-			RootMovieScene->SetPlaybackRange(0, RootDuration);
+		// Update playback range for the root sequence.
+		SetRootSequencePlaybackRange(TimePerCycle, *RootSequence);
 
-			// Iterate over subsequences and update their time scales.
-			for (UMovieSceneSection* Section : RootMovieScene->GetAllSections())
+		// Iterate over subsequences and update their time scales.
+		for (UMovieSceneSection* Section : RootMovieScene->GetAllSections())
+		{
+			if (UMovieSceneSubSection* SubSection = Cast<UMovieSceneSubSection>(Section))
 			{
-				if (UMovieSceneSubSection* SubSection = Cast<UMovieSceneSubSection>(Section))
-				{
-					UpdateSubSectionTimeScale(SubSection);
-				}
+				UpdateSubSectionTimeScale(SubSection);
 			}
+		}
 
+		if (bIsGameWorld)
+		{
 			// Update the sequence player frame range from the root sequence play range
 			const TRange<FFrameNumber> PlaybackRange = RootMovieScene->GetPlaybackRange();
+			const FFrameRate DisplayRate = RootMovieScene->GetDisplayRate();
+			const FFrameRate TickResolution = RootMovieScene->GetTickResolution();
+			
 			const FFrameNumber SrcStartFrame = UE::MovieScene::DiscreteInclusiveLower(PlaybackRange);
 			const FFrameNumber SrcEndFrame   = UE::MovieScene::DiscreteExclusiveUpper(PlaybackRange);
 
