DaySequence: Use PlayRate by default when changing playback speed during game.

Instead of changing TimePerCycle at runtime to change playback speed, we can now take advantage of PlayRate on the player.

TimePerCycle still functions the same, and disabling the CVar DaySequence.UsePlayRate will preserve old functionality, but now runtime changes are simpler and this resolves some flicker issues we used to have.

diff --git a/Engine/Plugins/Experimental/DaySequence/Source/DaySequence/Private/DaySequenceActor.cpp b/Engine/Plugins/Experimental/DaySequence/Source/DaySequence/Private/DaySequenceActor.cpp
--- a/Engine/Plugins/Experimental/DaySequence/Source/DaySequence/Private/DaySequenceActor.cpp
+++ b/Engine/Plugins/Experimental/DaySequence/Source/DaySequence/Private/DaySequenceActor.cpp
@@ -69,10 +69,17 @@
 	);
 
 	static TAutoConsoleVariable<bool> CVarDumpStackTraceOnSetTimeOfDay(
-	TEXT("DaySequence.DumpStackTraceOnSetTimeOfDay"),
-	false,
-	TEXT("When true, calls to SetTimeOfDay will dump the stack trace."),
-	ECVF_Default
+		TEXT("DaySequence.DumpStackTraceOnSetTimeOfDay"),
+		false,
+		TEXT("When true, calls to SetTimeOfDay will dump the stack trace."),
+		ECVF_Default
 	);
 
+	static TAutoConsoleVariable<bool> CVarUsePlayRate(
+		TEXT("DaySequence.UsePlayRate"),
+		true,
+		TEXT("When true, changing attempting to change TimePerCycle at runtime will instead adjust the player's play rate. When false, we properly change TimePerCycle."),
+		ECVF_Default
+	);
+
 	bool IsTimePerCycleOverflow(const FDaySequenceTime& InTimePerCycle, const UMovieScene& InMovieScene)
@@ -208,6 +215,11 @@
 	}
 	
 	InitializePlayer();
+
+	if (IConsoleVariable* UsePlayRateVariable = UE::DaySequence::CVarUsePlayRate.AsVariable())
+	{
+		UsePlayRateVariable->OnChangedDelegate().AddUObject(this, &ADaySequenceActor::OnUsePlayRateChanged);
+	}
 }
 
 #if WITH_EDITOR
@@ -339,6 +351,16 @@
 #endif
 }
 
+void ADaySequenceActor::BeginDestroy()
+{
+	Super::BeginDestroy();
+	
+	if (IConsoleVariable* UsePlayRateVariable = UE::DaySequence::CVarUsePlayRate.AsVariable())
+	{
+		UsePlayRateVariable->OnChangedDelegate().RemoveAll(this);
+	}
+}
+
 void ADaySequenceActor::BeginPlay()
 {
 	Super::BeginPlay();
@@ -968,6 +990,87 @@
 {
 }
 
+void ADaySequenceActor::Runtime_SetTimePerCycle_Impl(float InHours)
+{
+	using namespace UE::DaySequence;
+
+	const TOptional<FFrameRate> FrameRate = RootSequence ? RootSequence->GetMovieScene()->GetDisplayRate() : TOptional<FFrameRate>();
+
+	const bool bIsGameWorld = GetWorld()->IsGameWorld();
+	const float CurrentTimeOfDay = bIsGameWorld ? GetTimeOfDay() : 0.f;
+
+	// Set min cycle length to 1 second.
+	InHours = FMath::Max(InHours, FDaySequenceTime::FromSeconds(1.f).ToHours());
+	const FDaySequenceTime NewTimePerCycle = FDaySequenceTime::FromHours(InHours);
+	if (NewTimePerCycle != TimePerCycle)
+	{
+		UMovieScene* RootMovieScene = RootSequence->GetMovieScene();
+		
+		TimePerCycle = NewTimePerCycle;
+		
+		// Update playback range for the root sequence.
+		SetRootSequencePlaybackRange(TimePerCycle, *RootSequence);
+
+		// Iterate over subsequences and update their time scales.
+		for (UMovieSceneSection* Section : RootMovieScene->GetAllSections())
+		{
+			if (UMovieSceneSubSection* SubSection = Cast<UMovieSceneSubSection>(Section))
+			{
+				UpdateSubSectionTimeScale(SubSection);
+			}
+		}
+
+		if (bIsGameWorld)
+		{
+			// Update the sequence player frame range from the root sequence play range
+			const TRange<FFrameNumber> PlaybackRange = RootMovieScene->GetPlaybackRange();
+			const FFrameRate DisplayRate = RootMovieScene->GetDisplayRate();
+			const FFrameRate TickResolution = RootMovieScene->GetTickResolution();
+			
+			const FFrameNumber SrcStartFrame = UE::MovieScene::DiscreteInclusiveLower(PlaybackRange);
+			const FFrameNumber SrcEndFrame   = UE::MovieScene::DiscreteExclusiveUpper(PlaybackRange);
+
+			const FFrameTime EndingTime = ConvertFrameTime(SrcEndFrame, TickResolution, DisplayRate);
+
+			const FFrameNumber StartingFrame = ConvertFrameTime(SrcStartFrame, TickResolution, DisplayRate).FloorToFrame();
+			const FFrameNumber EndingFrame   = EndingTime.FloorToFrame();
+
+			if (UDaySequencePlayer* Player = GetSequencePlayerInternal())
+			{
+				Player->SetFrameRange(StartingFrame.Value, (EndingFrame - StartingFrame).Value, EndingTime.GetSubFrame());
+			}
+
+			// Finally update the play position to match the current time of day
+			SetTimeOfDay(CurrentTimeOfDay);
+		}
+	}
+}
+
+void ADaySequenceActor::OnUsePlayRateChanged(IConsoleVariable* UsePlayRateVariable)
+{
+	// This function really only makes sense if we are running client/server in the same process as it will be called for both actors.
+	// If we are not in a single process, this function will likely do more harm than good.
+	if (UDaySequencePlayer* Player = GetSequencePlayerInternal())
+	{
+		const ADaySequenceActor* CDO =  GetClass()->GetDefaultObject<ADaySequenceActor>();
+		
+		if (UsePlayRateVariable->GetBool())
+		{
+			// determine play rate from current TPC compared to CDO value, set TPC back to default, and set play rate
+			float EffectivePlayRate = CDO->TimePerCycle.ToHours() / TimePerCycle.ToHours();
+			Runtime_SetTimePerCycle_Impl(CDO->TimePerCycle.ToHours());
+			Player->SetPlayRate(EffectivePlayRate);
+		}
+		else
+		{
+			// determine TPC from current play rate, set TPC, set player play rate back to 1.f
+			float EffectiveTimePerCycleHours = CDO->TimePerCycle.ToHours() / Player->GetPlayRate();
+			Runtime_SetTimePerCycle_Impl(EffectiveTimePerCycleHours);
+			Player->SetPlayRate(1.f);
+		}
+	}
+}
+
 void ADaySequenceActor::StartDaySequenceUpdateTimer()
 {
 	if (GetNetMode() == NM_DedicatedServer)
@@ -1169,65 +1272,61 @@
 
 void ADaySequenceActor::SetTimePerCycle(float InHours)
 {
+	if (GetWorld()->IsGameWorld())
+	{
+		return;
+	}
+	
 	// Set min cycle length to 1 second.
 	InHours = FMath::Max(InHours, FDaySequenceTime::FromSeconds(1.f).ToHours());
 	TimePerCycle = FDaySequenceTime::FromHours(InHours);
 }
 
+void ADaySequenceActor::SetPlayRate_Implementation(float NewRate)
+{
+	if (UE::DaySequence::CVarUsePlayRate.GetValueOnAnyThread())
+	{
+		// we are assuming TimePerCycle is default here. should be enforced by CVar callback
+		if (UDaySequencePlayer* Player = GetSequencePlayerInternal())
+		{
+			Player->SetPlayRate(NewRate);
+		}
+	}
+	else
+	{
+		const ADaySequenceActor* CDO = GetClass()->GetDefaultObject<ADaySequenceActor>();
+		const float EffectiveTimePerCycleHours = CDO->TimePerCycle.ToHours() / NewRate;
+		Runtime_SetTimePerCycle_Impl(EffectiveTimePerCycleHours);
+	}
+}
+	
+float ADaySequenceActor::GetPlayRate() const
+{
+	if (UE::DaySequence::CVarUsePlayRate.GetValueOnAnyThread())
+	{
+		if (const UDaySequencePlayer* Player = GetSequencePlayerInternal())
+		{
+			return Player->GetPlayRate();
+		}	
+	}
+	else
+	{
+		const ADaySequenceActor* CDO = GetClass()->GetDefaultObject<ADaySequenceActor>();
+		return CDO->TimePerCycle.ToHours() / TimePerCycle.ToHours();
+	}
+	
+	return 1.f;
+}
+
 void ADaySequenceActor::Multicast_SetTimePerCycle_Implementation(float InHours)
 {
 	using namespace UE::DaySequence;
-	
-	const TOptional<FFrameRate> FrameRate = RootSequence ? RootSequence->GetMovieScene()->GetDisplayRate() : TOptional<FFrameRate>();
 
-	const bool bIsGameWorld = GetWorld()->IsGameWorld();
-	const float CurrentTimeOfDay = bIsGameWorld ? GetTimeOfDay() : 0.f;
+	const ADaySequenceActor* CDO = GetClass()->GetDefaultObject<ADaySequenceActor>();
+	const float DesiredPlayRate = CDO->TimePerCycle.ToHours() / InHours;
 
-	// Set min cycle length to 1 second.
-	InHours = FMath::Max(InHours, FDaySequenceTime::FromSeconds(1.f).ToHours());
-	const FDaySequenceTime NewTimePerCycle = FDaySequenceTime::FromHours(InHours);
-	if (NewTimePerCycle != TimePerCycle)
-	{
-		UMovieScene* RootMovieScene = RootSequence->GetMovieScene();
-		
-		TimePerCycle = NewTimePerCycle;
-		
-		// Update playback range for the root sequence.
-		SetRootSequencePlaybackRange(TimePerCycle, *RootSequence);
-
-		// Iterate over subsequences and update their time scales.
-		for (UMovieSceneSection* Section : RootMovieScene->GetAllSections())
-		{
-			if (UMovieSceneSubSection* SubSection = Cast<UMovieSceneSubSection>(Section))
-			{
-				UpdateSubSectionTimeScale(SubSection);
-			}
-		}
-
-		if (bIsGameWorld)
-		{
-			// Update the sequence player frame range from the root sequence play range
-			const TRange<FFrameNumber> PlaybackRange = RootMovieScene->GetPlaybackRange();
-			const FFrameRate DisplayRate = RootMovieScene->GetDisplayRate();
-			const FFrameRate TickResolution = RootMovieScene->GetTickResolution();
-			
-			const FFrameNumber SrcStartFrame = UE::MovieScene::DiscreteInclusiveLower(PlaybackRange);
-			const FFrameNumber SrcEndFrame   = UE::MovieScene::DiscreteExclusiveUpper(PlaybackRange);
-
-			const FFrameTime EndingTime = ConvertFrameTime(SrcEndFrame, TickResolution, DisplayRate);
-
-			const FFrameNumber StartingFrame = ConvertFrameTime(SrcStartFrame, TickResolution, DisplayRate).FloorToFrame();
-			const FFrameNumber EndingFrame   = EndingTime.FloorToFrame();
-
-			if (UDaySequencePlayer* Player = GetSequencePlayerInternal())
-			{
-				Player->SetFrameRange(StartingFrame.Value, (EndingFrame - StartingFrame).Value, EndingTime.GetSubFrame());
-			}
-
-			// Finally update the play position to match the current time of day
-			SetTimeOfDay(CurrentTimeOfDay);
-		}
-	}
+	// We are explicitly calling the implementation because this function itself is already NetMulticast.
+	SetPlayRate_Implementation(DesiredPlayRate);
 }
 
 float ADaySequenceActor::GetInitialTimeOfDay() const
diff --git a/Engine/Plugins/Experimental/DaySequence/Source/DaySequence/Private/DaySequenceCheatManagerExtension.cpp b/Engine/Plugins/Experimental/DaySequence/Source/DaySequence/Private/DaySequenceCheatManagerExtension.cpp
--- a/Engine/Plugins/Experimental/DaySequence/Source/DaySequence/Private/DaySequenceCheatManagerExtension.cpp
+++ b/Engine/Plugins/Experimental/DaySequence/Source/DaySequence/Private/DaySequenceCheatManagerExtension.cpp
@@ -37,11 +37,7 @@ void UDaySequenceCheatManagerExtension::SetTimeOfDay(float NewTimeOfDay) const
 
 void UDaySequenceCheatManagerExtension::SetTimeOfDaySpeed(float NewTimeOfDaySpeedMultiplier) const
 {
-	ADaySequenceActor* DaySequenceActor = GetDaySequenceActor();
-	const UClass* Class = DaySequenceActor ? DaySequenceActor->GetClass() : nullptr;
-	const ADaySequenceActor* CDO = Class ? Class->GetDefaultObject<ADaySequenceActor>() : nullptr;
-	
-	if (DaySequenceActor && CDO)
+	if (ADaySequenceActor* DaySequenceActor = GetDaySequenceActor())
 	{
 		if (NewTimeOfDaySpeedMultiplier < 0.f)
 		{
@@ -54,8 +50,7 @@ void UDaySequenceCheatManagerExtension::SetTimeOfDaySpeed(float NewTimeOfDaySpee
 		}
 		else
 		{
-			const float DesiredTimePerCycle = CDO->GetTimePerCycle() / NewTimeOfDaySpeedMultiplier;
-			DaySequenceActor->Multicast_SetTimePerCycle(DesiredTimePerCycle);
+			DaySequenceActor->SetPlayRate(NewTimeOfDaySpeedMultiplier);
 			DaySequenceActor->Play();
 		}
 	}
diff --git a/Engine/Plugins/Experimental/DaySequence/Source/DaySequence/Private/DaySequencePlayer.cpp b/Engine/Plugins/Experimental/DaySequence/Source/DaySequence/Private/DaySequencePlayer.cpp
--- a/Engine/Plugins/Experimental/DaySequence/Source/DaySequence/Private/DaySequencePlayer.cpp
+++ b/Engine/Plugins/Experimental/DaySequence/Source/DaySequence/Private/DaySequencePlayer.cpp
@@ -690,6 +690,12 @@ float UDaySequencePlayer::GetPlayRate() const
 void UDaySequencePlayer::SetPlayRate(float PlayRate)
 {
 	PlaybackSettings.PlayRate = PlayRate;
+
+	if (HasAuthority())
+	{
+		// Explicitly notify clients of change so they can discard server time samples.
+		RPC_OnPlayRateChanged();
+	}
 }
 
 FFrameTime UDaySequencePlayer::GetLastValidTime() const
@@ -1590,6 +1596,11 @@ void UDaySequencePlayer::RPC_OnFinishPlaybackEvent_Implementation(FFrameTime Sto
 	FinishPlaybackInternal(StoppedTime);
 }
 
+void UDaySequencePlayer::RPC_OnPlayRateChanged_Implementation()
+{
+	ServerTimeSamples.Reset();
+}
+
 void UDaySequencePlayer::PostNetReceive()
 {
 	// ~~~~~~~~~~~~~~~~~~~~~~~~~~~~~~~~~~~~~~~~~~~~~~~~~~~~~~~~~~~~~~~~~~~~~~~~~~~~~~~~~~~
diff --git a/Engine/Plugins/Experimental/DaySequence/Source/DaySequence/Private/DaySequencePlayer.h b/Engine/Plugins/Experimental/DaySequence/Source/DaySequence/Private/DaySequencePlayer.h
--- a/Engine/Plugins/Experimental/DaySequence/Source/DaySequence/Private/DaySequencePlayer.h
+++ b/Engine/Plugins/Experimental/DaySequence/Source/DaySequence/Private/DaySequencePlayer.h
@@ -466,6 +466,12 @@ private:
 	UFUNCTION(netmulticast, reliable)
 	void RPC_OnFinishPlaybackEvent(FFrameTime StoppedTime, int32 NewSerialNumber);
 
+	/**
+	 * Called on the server when PlaybackSettings.PlayRate changes so that clients can discard server time samples.
+	 */
+	UFUNCTION(netmulticast, reliable)
+	void RPC_OnPlayRateChanged();
+	
 	/**
 	 * Check whether this sequence player is an authority, as determined by its outer Actor
 	 */
diff --git a/Engine/Plugins/Experimental/DaySequence/Source/DaySequence/Public/DaySequenceActor.h b/Engine/Plugins/Experimental/DaySequence/Source/DaySequence/Public/DaySequenceActor.h
--- a/Engine/Plugins/Experimental/DaySequence/Source/DaySequence/Public/DaySequenceActor.h
+++ b/Engine/Plugins/Experimental/DaySequence/Source/DaySequence/Public/DaySequenceActor.h
@@ -170,7 +170,7 @@ public:
 	void SetDayLength(float InHours);
 
 	/**
-	 * Get the duration of each day cycle in hours.
+	 * Get the duration of each day cycle in hours (assuming PlayRate is 1.0)
 	 *
 	 * @return float, time in hours
 	 */
@@ -181,17 +181,34 @@ public:
 	 * Set the duration of a day cycle in real time in hours.
 	 *
 	 * This is currently intentionally not exposed to Blueprint since
-	 * this property can only be realized by rebuilding the root sequence.
+	 * the effect is achieved by changing PlayRate at runtime.
 	 *
 	 * @param InHours time per cycle in hours.
 	 */
 	void SetTimePerCycle(float InHours);
 
+	/**
+	 * Sets the play rate of the day cycle at runtime.
+	 * A play rate of 2.0 will result in a day cycle that is half as long as Time Per Cycle.
+	 *
+	 * @param NewRate The new play rate.
+	 */
+	UFUNCTION(BlueprintCallable, NetMulticast, Reliable, Category=RuntimeDayCycle)
+	void SetPlayRate(float NewRate);
+
+	/**
+	 * Gets the play rate of the day cycle at runtime. Always 1 in editor worlds.
+	 *
+	 * @return The current play rate.
+	 */
+	UFUNCTION(BlueprintCallable, Category=RuntimeDayCycle)
+	float GetPlayRate() const;
+	
 	/**
 	 * Set the duration of a day cycle in real time in hours and update
 	 * the root sequence in game.
 	 *
-	 * This method is the the network multicast variant of SetTimePerCycle.
+	 * This method is the network multicast variant of SetTimePerCycle.
 	 * In game, this will properly update the root sequence and subsequence
 	 * time scales on both the server and the client. These changes must
 	 * be invoked on both server and client since only playback position &
@@ -199,6 +216,7 @@ public:
 	 *
 	 * @param InHours time per cycle in hours.
 	 */
+	UE_DEPRECATED(5.6, "Multicast_SetTimePerCycle is deprecated, use SetPlayRate instead.")
 	UFUNCTION(NetMulticast, Reliable)
 	void Multicast_SetTimePerCycle(float InHours);
 
@@ -364,6 +382,7 @@ protected:
 	//~ Begin UObject interface
 	virtual void PostInitProperties() override;
 	virtual void PostLoad() override;
+	virtual void BeginDestroy() override;
 #if WITH_EDITOR
 	virtual void PostEditChangeProperty(FPropertyChangedEvent& PropertyChangedEvent) override;
 #endif
@@ -435,6 +454,9 @@ protected:
 	void OnSequencePlayerUpdate(const UDaySequencePlayer& Player, FFrameTime CurrentTime, FFrameTime PreviousTime);
 	virtual void SequencePlayerUpdated(float CurrentTime, float PreviousTime);
 
+	/* Callback for changes to the DaySequence.UsePlayRate cvar. Keeps actor in a good state when the CVar changes. */
+	void OnUsePlayRateChanged(IConsoleVariable* UsePlayRateVariable);
+	
 public:
 	/**
 	 * User-provided interpolation curve that maps day cycle times to desired cycle times (usually from 0-24 hours).
@@ -592,5 +614,9 @@ protected:
 	
 	UPROPERTY(Transient)
 	TMap<TSubclassOf<UDaySequenceConditionTag>, TObjectPtr<UDaySequenceConditionTag>> TrackConditionMap;
+
+private:
+	/* Used to change TimePerCycle while a player is active. Must be called on server and all clients in order to prevent synchronization issues. */
+	void Runtime_SetTimePerCycle_Impl(float InHours);
 };
 
