DaySequence: Move FOverrideUpdateIntervalHandle out of private header.

diff --git a/Engine/Plugins/Experimental/DaySequence/Source/DaySequence/Private/DaySequencePlayer.h b/Engine/Plugins/Experimental/DaySequence/Source/DaySequence/Private/DaySequencePlayer.h
--- a/Engine/Plugins/Experimental/DaySequence/Source/DaySequence/Private/DaySequencePlayer.h
+++ b/Engine/Plugins/Experimental/DaySequence/Source/DaySequence/Private/DaySequencePlayer.h
@@ -31,24 +31,6 @@ class FSequenceWeights;
 DECLARE_DYNAMIC_MULTICAST_DELEGATE(FOnDaySequencePlayerEvent);
 DECLARE_DELEGATE(FOnDaySequencePlayerNativeEvent);
 
-namespace UE::DaySequence
-{
-class FOverrideUpdateIntervalHandle
-{
-public:
-	FOverrideUpdateIntervalHandle(UDaySequencePlayer* InPlayer);
-	~FOverrideUpdateIntervalHandle();
-	
-	void StartOverriding();
-	void StopOverriding();
-
-private:
-	TWeakObjectPtr<UDaySequencePlayer> WeakPlayer;
-	
-	bool bIsOverriding;
-};
-}
-
 USTRUCT()
 struct FDaySequencePlaybackParams
 {
diff --git a/Engine/Plugins/Experimental/DaySequence/Source/DaySequence/Public/IDaySequencePlayer.h b/Engine/Plugins/Experimental/DaySequence/Source/DaySequence/Public/IDaySequencePlayer.h
--- a/Engine/Plugins/Experimental/DaySequence/Source/DaySequence/Public/IDaySequencePlayer.h
+++ b/Engine/Plugins/Experimental/DaySequence/Source/DaySequence/Public/IDaySequencePlayer.h
@@ -4,10 +4,26 @@
 
 #include "Misc/QualifiedFrameTime.h"
 #include "Templates/SharedPointer.h"
+#include "UObject/WeakObjectPtr.h"
+
+class UDaySequencePlayer;
 
 namespace UE::DaySequence
 {
-class FOverrideUpdateIntervalHandle;
+class DAYSEQUENCE_API FOverrideUpdateIntervalHandle
+{
+public:
+	FOverrideUpdateIntervalHandle(UDaySequencePlayer* InPlayer);
+	~FOverrideUpdateIntervalHandle();
+
+	void StartOverriding();
+	void StopOverriding();
+
+private:
+	TWeakObjectPtr<UDaySequencePlayer> WeakPlayer;
+
+	bool bIsOverriding;
+};
 }
 
 class DAYSEQUENCE_API IDaySequencePlayer
