DaySequenceActor: Soft backout of dynamic tick resolution change

diff --git a/Engine/Plugins/Experimental/DaySequence/Source/DaySequence/Private/DaySequenceActor.cpp b/Engine/Plugins/Experimental/DaySequence/Source/DaySequence/Private/DaySequenceActor.cpp
--- a/Engine/Plugins/Experimental/DaySequence/Source/DaySequence/Private/DaySequenceActor.cpp
+++ b/Engine/Plugins/Experimental/DaySequence/Source/DaySequence/Private/DaySequenceActor.cpp
@@ -82,6 +82,13 @@ namespace UE::DaySequence
 		ECVF_Default
 	);
 
+	static TAutoConsoleVariable<bool> CVarDynamicTickResolution(
+		TEXT("DaySequence.DynamicTickResolution"),
+		false,
+		TEXT("When true, attempts to resolve TPC integer overflow by decreasing tick resolution of the Root Sequence.."),
+		ECVF_Default
+	);
+
 	bool IsTimePerCycleOverflow(const FDaySequenceTime& InTimePerCycle, const UMovieScene& InMovieScene)
 	{
 		const FFrameRate DisplayRate = InMovieScene.GetDisplayRate();
@@ -112,7 +119,7 @@ namespace UE::DaySequence
 		if (UMovieScene* RootMovieScene = InRootSequence.GetMovieScene())
 		{
 			// Drop tick resolution if we are going to overflow playback range at this TPC
-			if (IsTimePerCycleOverflow(InTimePerCycle, *RootMovieScene))
+			if (CVarDynamicTickResolution.GetValueOnAnyThread() && IsTimePerCycleOverflow(InTimePerCycle, *RootMovieScene))
 			{
 				RootMovieScene->SetTickResolutionDirectly(GetMaximumTickResolution(InTimePerCycle, *RootMovieScene));
 			}
@@ -1005,6 +1012,12 @@ void ADaySequenceActor::Runtime_SetTimePerCycle_Impl(float InHours)
 	if (NewTimePerCycle != TimePerCycle)
 	{
 		UMovieScene* RootMovieScene = RootSequence->GetMovieScene();
+
+		if (!CVarDynamicTickResolution.GetValueOnAnyThread() && IsTimePerCycleOverflow(NewTimePerCycle, *RootMovieScene))
+		{
+			UE_LOG(LogDaySequence, Warning, TEXT("Skipping SetTimePerCycle( %f hours ) to avoid overflowing playback range"), InHours);
+			return;
+		}
 		
 		TimePerCycle = NewTimePerCycle;
 		
@@ -1272,11 +1285,6 @@ float ADaySequenceActor::GetTimePerCycle() const
 
 void ADaySequenceActor::SetTimePerCycle(float InHours)
 {
-	if (GetWorld()->IsGameWorld())
-	{
-		return;
-	}
-	
 	// Set min cycle length to 1 second.
 	InHours = FMath::Max(InHours, FDaySequenceTime::FromSeconds(1.f).ToHours());
 	TimePerCycle = FDaySequenceTime::FromHours(InHours);
