DaySequence: Choose arbitrary RootDuration and use PlayRate to realize TimePerCycle.

RootDuration is now computed so that there is some visual mapping of the RootSequence's playback range from the user's perspective to DayLength. For example, if DayLength is 24:00:00, the user will see that the playback range appears to be 2400 frames, which makes it convenient to reason about what time a particular frame corresponds to. Subsequences continue to set their own playback range.

SetPlayRate behaves as one would expect: a play rate of 1 means that the day cycle takes 1 TimePerCycle, a play rate of 2 means that the day cycle takes 1/2 TimePerCycle, etc etc. Under the hood, we are doing this by determining what play rate the player needs to operate at in order for the rate * the player's duration to match the DSA' TimePerCycle. At runtime, we don't change TimePerCycle (at least it isn't necessary or intended). Multicast_SetTimePerCycle continues to route through to SetPlayRate.

diff --git a/Engine/Plugins/Experimental/DaySequence/Source/DaySequence/Private/DaySequenceActor.cpp b/Engine/Plugins/Experimental/DaySequence/Source/DaySequence/Private/DaySequenceActor.cpp
--- a/Engine/Plugins/Experimental/DaySequence/Source/DaySequence/Private/DaySequenceActor.cpp
+++ b/Engine/Plugins/Experimental/DaySequence/Source/DaySequence/Private/DaySequenceActor.cpp
@@ -74,64 +74,6 @@ namespace UE::DaySequence
 		TEXT("When true, calls to SetTimeOfDay will dump the stack trace."),
 		ECVF_Default
 	);
-
-	static TAutoConsoleVariable<bool> CVarUsePlayRate(
-		TEXT("DaySequence.UsePlayRate"),
-		true,
-		TEXT("When true, changing attempting to change TimePerCycle at runtime will instead adjust the player's play rate. When false, we properly change TimePerCycle."),
-		ECVF_Default
-	);
-
-	static TAutoConsoleVariable<bool> CVarDynamicTickResolution(
-		TEXT("DaySequence.DynamicTickResolution"),
-		false,
-		TEXT("When true, attempts to resolve TPC integer overflow by decreasing tick resolution of the Root Sequence.."),
-		ECVF_Default
-	);
-
-	bool IsTimePerCycleOverflow(const FDaySequenceTime& InTimePerCycle, const UMovieScene& InMovieScene)
-	{
-		const FFrameRate DisplayRate = InMovieScene.GetDisplayRate();
-		const FFrameRate TickResolution = InMovieScene.GetTickResolution();
-		const float TimePerCycleSeconds = InTimePerCycle.ToSeconds();
-		const int32 CycleSeconds = static_cast<int32>(TimePerCycleSeconds) + 1;
-		const int32 DisplayFactor = DisplayRate.Numerator / DisplayRate.Denominator;
-		const int32 TickFactor = TickResolution.Numerator / TickResolution.Denominator;
-		const int32 NewCycleSeconds = CycleSeconds * DisplayFactor * TickFactor;
-		const int32 CheckCycleSeconds = NewCycleSeconds / (DisplayFactor * TickFactor);
-		return (CheckCycleSeconds != CycleSeconds);
-	}
-
-	FFrameRate GetMaximumTickResolution(const FDaySequenceTime& InTimePerCycle, const UMovieScene& InMovieScene)
-	{
-		const FFrameRate DisplayRate = InMovieScene.GetDisplayRate();
-		const int32 DisplayFactor = DisplayRate.Numerator / DisplayRate.Denominator;
-		int32 MaxResolution = std::numeric_limits<int32>::max() / (InTimePerCycle.ToSeconds() * DisplayFactor);
-
-		// Set MaxResolution to the closest multiple of DisplayFactor that is less than the initially computed value of MaxResolution.
-		MaxResolution = (MaxResolution / DisplayFactor) * DisplayFactor;
-		
-		return FFrameRate(MaxResolution, 1);
-	}
-
-	void SetRootSequencePlaybackRange(const FDaySequenceTime& InTimePerCycle, const UMovieSceneSequence& InRootSequence)
-	{
-		if (UMovieScene* RootMovieScene = InRootSequence.GetMovieScene())
-		{
-			// Drop tick resolution if we are going to overflow playback range at this TPC
-			if (CVarDynamicTickResolution.GetValueOnAnyThread() && IsTimePerCycleOverflow(InTimePerCycle, *RootMovieScene))
-			{
-				RootMovieScene->SetTickResolutionDirectly(GetMaximumTickResolution(InTimePerCycle, *RootMovieScene));
-			}
-			
-			const int32 RootDuration = RootMovieScene->GetTickResolution().AsFrameNumber(InTimePerCycle.ToSeconds()).Value;
-			RootMovieScene->SetPlaybackRange(0, RootDuration);
-#if WITH_EDITOR
-			RootMovieScene->SetPlaybackRangeLocked(true);
-#endif
-			RootMovieScene->MarkAsChanged();
-		}
-	}
 } // namespace UE::DaySequence
 
 ADaySequenceActor::ADaySequenceActor(const FObjectInitializer& Init)
@@ -222,11 +164,6 @@ void ADaySequenceActor::PostInitializeComponents()
 	}
 	
 	InitializePlayer();
-
-	if (IConsoleVariable* UsePlayRateVariable = UE::DaySequence::CVarUsePlayRate.AsVariable())
-	{
-		UsePlayRateVariable->OnChangedDelegate().AddUObject(this, &ADaySequenceActor::OnUsePlayRateChanged);
-	}
 }
 
 #if WITH_EDITOR
@@ -361,11 +298,6 @@ void ADaySequenceActor::PostLoad()
 void ADaySequenceActor::BeginDestroy()
 {
 	Super::BeginDestroy();
-	
-	if (IConsoleVariable* UsePlayRateVariable = UE::DaySequence::CVarUsePlayRate.AsVariable())
-	{
-		UsePlayRateVariable->OnChangedDelegate().RemoveAll(this);
-	}
 }
 
 void ADaySequenceActor::BeginPlay()
@@ -523,6 +455,10 @@ void ADaySequenceActor::PostEditChangeProperty(FPropertyChangedEvent& PropertyCh
 	else if (PropertyName == GET_MEMBER_NAME_CHECKED(ADaySequenceActor, DayLength))
 	{
 		SetDayLength(GetDayLength());
+
+		// RootSequence playback range is computed using DayLength, so lets reconstruct RootSequence to reflect this change.
+		RootSequence = nullptr;
+		ReinitializeActor();
 	}
 	else if (PropertyName == GET_MEMBER_NAME_CHECKED(ADaySequenceActor, TimePerCycle))
 	{
@@ -645,7 +581,7 @@ void ADaySequenceActor::InitializeRootSequence()
 		RootSequence->Initialize(RF_Transient);
 		RootSequence->SetSequenceFlags(EMovieSceneSequenceFlags::Volatile);
 
-		SetRootSequencePlaybackRange(TimePerCycle, *RootSequence);
+		SetRootSequencePlaybackRange();
 #if ROOT_SEQUENCE_RECONSTRUCTION_ENABLED
 	}
 #endif
@@ -736,6 +672,50 @@ void ADaySequenceActor::InitializeRootSequence()
 	OnPostRootSequenceChanged.Broadcast();
 }
 
+void ADaySequenceActor::SetRootSequencePlaybackRange()
+{
+	if (UMovieScene* RootMovieScene = RootSequence->GetMovieScene())
+	{
+		// Compute RootDuration so that the frame range as viewed by the user is [0, 100 * DayLengthHours].
+		// For example, if DayLength is 24:00:00, frame range is [0, 2400].
+		
+		const FFrameRate TickDisplayRatio = RootMovieScene->GetTickResolution() / RootMovieScene->GetDisplayRate();
+		const int32 HoursToFrames = 100 * TickDisplayRatio.Numerator / TickDisplayRatio.Denominator;
+		const int32 RootDuration = FMath::Min(HoursToFrames * DayLength.ToHours(), 100000000);
+		
+		RootMovieScene->SetPlaybackRange(0, RootDuration);
+#if WITH_EDITOR
+		RootMovieScene->SetPlaybackRangeLocked(true);
+#endif
+		
+		RootMovieScene->MarkAsChanged();
+	}
+}
+
+float ADaySequenceActor::GetBasePlayRate() const
+{
+	if (!RootSequence)
+	{
+		return 1.f;
+	}
+	
+	const float PlayerDurationSeconds = GetBaseDuration() * FDaySequenceTime::SecondsPerHour;
+	return PlayerDurationSeconds / TimePerCycle.ToSeconds();
+}
+
+float ADaySequenceActor::GetBaseDuration() const
+{
+	if (!RootSequence)
+	{
+		return TimePerCycle.ToHours();
+	}
+
+	const UMovieScene* MovieScene = RootSequence->GetMovieScene();
+	const float BaseDurationSeconds = MovieScene->GetTickResolution().AsSeconds(MovieScene->GetPlaybackRange().GetUpperBoundValue());
+
+	return BaseDurationSeconds / FDaySequenceTime::SecondsPerHour;
+}
+
 #if ROOT_SEQUENCE_RECONSTRUCTION_ENABLED
 bool ADaySequenceActor::MarkDaySequences(FSubSectionPreserveMap* SectionsToPreserve)
 {
@@ -977,10 +957,11 @@ void ADaySequenceActor::OnSequencePlayerUpdate(const UDaySequencePlayer& Player,
 	
 	auto FrameTimeToDayHours = [this](const FFrameTime& Time)
 	{
-		const FFrameRate FrameRate = RootSequence->GetMovieScene()->GetDisplayRate();
+		const UMovieScene* MovieScene = RootSequence->GetMovieScene();
+		const FFrameRate FrameRate = MovieScene->GetDisplayRate();
 		const double CurrentTimeSeconds = FQualifiedFrameTime(Time, FrameRate).AsSeconds();
-		const float DayCycleSeconds = TimePerCycle.ToSeconds();
-		const float SequenceRatio = CurrentTimeSeconds / DayCycleSeconds;
+		const float PlayerDurationSeconds = GetBaseDuration() * FDaySequenceTime::SecondsPerHour;
+		const float SequenceRatio = CurrentTimeSeconds / PlayerDurationSeconds;
 		return DayLength.ToHours() * SequenceRatio;
 	};
 	const float CurrentHours = FrameTimeToDayHours(CurrentTime);
@@ -997,93 +978,6 @@ void ADaySequenceActor::SequencePlayerUpdated(float CurrentTime, float PreviousT
 {
 }
 
-void ADaySequenceActor::Runtime_SetTimePerCycle_Impl(float InHours)
-{
-	using namespace UE::DaySequence;
-
-	const TOptional<FFrameRate> FrameRate = RootSequence ? RootSequence->GetMovieScene()->GetDisplayRate() : TOptional<FFrameRate>();
-
-	const bool bIsGameWorld = GetWorld()->IsGameWorld();
-	const float CurrentTimeOfDay = bIsGameWorld ? GetTimeOfDay() : 0.f;
-
-	// Set min cycle length to 1 second.
-	InHours = FMath::Max(InHours, FDaySequenceTime::FromSeconds(1.f).ToHours());
-	const FDaySequenceTime NewTimePerCycle = FDaySequenceTime::FromHours(InHours);
-	if (NewTimePerCycle != TimePerCycle)
-	{
-		UMovieScene* RootMovieScene = RootSequence->GetMovieScene();
-
-		if (!CVarDynamicTickResolution.GetValueOnAnyThread() && IsTimePerCycleOverflow(NewTimePerCycle, *RootMovieScene))
-		{
-			UE_LOG(LogDaySequence, Warning, TEXT("Skipping SetTimePerCycle( %f hours ) to avoid overflowing playback range"), InHours);
-			return;
-		}
-		
-		TimePerCycle = NewTimePerCycle;
-		
-		// Update playback range for the root sequence.
-		SetRootSequencePlaybackRange(TimePerCycle, *RootSequence);
-
-		// Iterate over subsequences and update their time scales.
-		for (UMovieSceneSection* Section : RootMovieScene->GetAllSections())
-		{
-			if (UMovieSceneSubSection* SubSection = Cast<UMovieSceneSubSection>(Section))
-			{
-				UpdateSubSectionTimeScale(SubSection);
-			}
-		}
-
-		if (bIsGameWorld)
-		{
-			// Update the sequence player frame range from the root sequence play range
-			const TRange<FFrameNumber> PlaybackRange = RootMovieScene->GetPlaybackRange();
-			const FFrameRate DisplayRate = RootMovieScene->GetDisplayRate();
-			const FFrameRate TickResolution = RootMovieScene->GetTickResolution();
-			
-			const FFrameNumber SrcStartFrame = UE::MovieScene::DiscreteInclusiveLower(PlaybackRange);
-			const FFrameNumber SrcEndFrame   = UE::MovieScene::DiscreteExclusiveUpper(PlaybackRange);
-
-			const FFrameTime EndingTime = ConvertFrameTime(SrcEndFrame, TickResolution, DisplayRate);
-
-			const FFrameNumber StartingFrame = ConvertFrameTime(SrcStartFrame, TickResolution, DisplayRate).FloorToFrame();
-			const FFrameNumber EndingFrame   = EndingTime.FloorToFrame();
-
-			if (UDaySequencePlayer* Player = GetSequencePlayerInternal())
-			{
-				Player->SetFrameRange(StartingFrame.Value, (EndingFrame - StartingFrame).Value, EndingTime.GetSubFrame());
-			}
-
-			// Finally update the play position to match the current time of day
-			SetTimeOfDay(CurrentTimeOfDay);
-		}
-	}
-}
-
-void ADaySequenceActor::OnUsePlayRateChanged(IConsoleVariable* UsePlayRateVariable)
-{
-	// This function really only makes sense if we are running client/server in the same process as it will be called for both actors.
-	// If we are not in a single process, this function will likely do more harm than good.
-	if (UDaySequencePlayer* Player = GetSequencePlayerInternal())
-	{
-		const ADaySequenceActor* CDO =  GetClass()->GetDefaultObject<ADaySequenceActor>();
-		
-		if (UsePlayRateVariable->GetBool())
-		{
-			// determine play rate from current TPC compared to CDO value, set TPC back to default, and set play rate
-			float EffectivePlayRate = CDO->TimePerCycle.ToHours() / TimePerCycle.ToHours();
-			Runtime_SetTimePerCycle_Impl(CDO->TimePerCycle.ToHours());
-			Player->SetPlayRate(EffectivePlayRate);
-		}
-		else
-		{
-			// determine TPC from current play rate, set TPC, set player play rate back to 1.f
-			float EffectiveTimePerCycleHours = CDO->TimePerCycle.ToHours() / Player->GetPlayRate();
-			Runtime_SetTimePerCycle_Impl(EffectiveTimePerCycleHours);
-			Player->SetPlayRate(1.f);
-		}
-	}
-}
-
 void ADaySequenceActor::StartDaySequenceUpdateTimer()
 {
 	if (GetNetMode() == NM_DedicatedServer)
@@ -1200,15 +1094,16 @@ void ADaySequenceActor::WarpEvaluationRange(FMovieSceneEvaluationRange& InOutRan
 	{
 		return;
 	}
-
+	
 	const FFrameRate TickRate   = InOutRange.GetFrameRate();
-	const float DayCycleSeconds = TimePerCycle.ToSeconds();
+	const UMovieScene* MovieScene = RootSequence->GetMovieScene();
+	const float DayCycleSeconds = MovieScene->GetTickResolution().AsSeconds(MovieScene->GetPlaybackRange().GetUpperBoundValue());
 	const float DayLengthHours  = DayLength.ToHours();
 
 	TRange<FFrameTime> Range = InOutRange.GetRange();
 
 	// Auto bounds checking
-	auto TrySetBounds = [&Range](FFrameTime LowerBound, FFrameTime UpperBound)
+	auto TrySetBounds = [&Range, MovieScene](FFrameTime LowerBound, FFrameTime UpperBound)
 	{
 		// Warp the lower bound if possible
 		if (Range.GetLowerBound().IsClosed())
@@ -1220,8 +1115,9 @@ void ADaySequenceActor::WarpEvaluationRange(FMovieSceneEvaluationRange& InOutRan
 		// Warp the upper bound if possible
 		if (Range.GetUpperBound().IsClosed())
 		{
-			// Set the upper bound value while retaining the inclusivity
-			Range.SetUpperBoundValue(UpperBound);
+			// Set the upper bound value while retaining the inclusivity.
+			// We are clamping to [0, PlaybackRangeUpperBound - 1] because our playback range has an exclusive upper bound.
+			Range.SetUpperBoundValue(FMath::Clamp(UpperBound, FFrameTime(0), FFrameTime(MovieScene->GetPlaybackRange().GetUpperBoundValue() - 1)));
 		}
 	};
 
@@ -1229,9 +1125,13 @@ void ADaySequenceActor::WarpEvaluationRange(FMovieSceneEvaluationRange& InOutRan
 	if (HasStaticTimeOfDay())
 	{
 		const float StaticTimeInGameHours = GetStaticTimeOfDay();
-		
-		const FFrameTime LowerBound = (StaticTimeInGameHours * DayCycleSeconds / DayLengthHours) * TickRate;
 
+		// The approach here:
+		// 1) converts static time to normalized time (StaticTimeInGameHours / DayLengthHours)
+		// 2) converts normalized time to sequencer time (X * DayCycleSeconds)
+		// 3) converts sequencer time to a frame time (X * TickRate)
+
+		const FFrameTime LowerBound = (StaticTimeInGameHours * DayCycleSeconds / DayLengthHours) * TickRate;
 		const FFrameTime UpperBound = (StaticTimeInGameHours * DayCycleSeconds / DayLengthHours) * TickRate;
 
 		TrySetBounds(LowerBound, UpperBound);
@@ -1292,37 +1192,22 @@ void ADaySequenceActor::SetTimePerCycle(float InHours)
 
 void ADaySequenceActor::SetPlayRate_Implementation(float NewRate)
 {
-	if (UE::DaySequence::CVarUsePlayRate.GetValueOnAnyThread())
-	{
-		// we are assuming TimePerCycle is default here. should be enforced by CVar callback
-		if (UDaySequencePlayer* Player = GetSequencePlayerInternal())
-		{
-			Player->SetPlayRate(NewRate);
-		}
-	}
-	else
+	if (UDaySequencePlayer* Player = GetSequencePlayerInternal())
 	{
-		const ADaySequenceActor* CDO = GetClass()->GetDefaultObject<ADaySequenceActor>();
-		const float EffectiveTimePerCycleHours = CDO->TimePerCycle.ToHours() / NewRate;
-		Runtime_SetTimePerCycle_Impl(EffectiveTimePerCycleHours);
+		// we are given a rate w.r.t. DSA time, not player time.
+		Player->SetPlayRate(NewRate * GetBasePlayRate());
 	}
 }
 	
 float ADaySequenceActor::GetPlayRate() const
 {
-	if (UE::DaySequence::CVarUsePlayRate.GetValueOnAnyThread())
-	{
-		if (const UDaySequencePlayer* Player = GetSequencePlayerInternal())
-		{
-			return Player->GetPlayRate();
-		}	
-	}
-	else
+	// Game world
+	if (const UDaySequencePlayer* Player = GetSequencePlayerInternal())
 	{
-		const ADaySequenceActor* CDO = GetClass()->GetDefaultObject<ADaySequenceActor>();
-		return CDO->TimePerCycle.ToHours() / TimePerCycle.ToHours();
+		return Player->GetPlayRate() / GetBasePlayRate();
 	}
-	
+
+	// Fallback
 	return 1.f;
 }
 
@@ -1357,8 +1242,8 @@ float ADaySequenceActor::GetTimeOfDay() const
 	{
 		const FQualifiedFrameTime CurrentFrameTime = Player->GetCurrentTime();
 		const double CurrentTimeSeconds = CurrentFrameTime.AsSeconds();
-		const float DayCycleSeconds = TimePerCycle.ToSeconds();
-		const float SequenceRatio = CurrentTimeSeconds / DayCycleSeconds;
+		const float PlayerDurationSeconds = GetBaseDuration() * FDaySequenceTime::SecondsPerHour;
+		const float SequenceRatio = CurrentTimeSeconds / PlayerDurationSeconds;
 		Result = DayLength.ToHours() * SequenceRatio;
 	}
 	else
@@ -1394,11 +1279,13 @@ bool ADaySequenceActor::SetTimeOfDay(float InHours)
 	UDaySequencePlayer* Player = GetSequencePlayerInternal();
 	if (HasValidRootSequence() && Player && World && World->IsGameWorld())
 	{
+		const UMovieScene* MovieScene = RootSequence->GetMovieScene();
+		const FFrameRate FrameRate = MovieScene->GetDisplayRate();
+		
 		// Convert the day time to sequence time
-		const FFrameRate FrameRate = RootSequence->GetMovieScene()->GetDisplayRate();
-		const float DayLengthHours = GetDayLength();
-		const float DayLengthRatio = FMath::Frac(InHours / DayLengthHours);
-		const float DayCycleSeconds = TimePerCycle.ToSeconds() * DayLengthRatio;
+		const float PlayerDurationSeconds = GetBaseDuration() * FDaySequenceTime::SecondsPerHour;
+		const float DayLengthRatio = FMath::Frac(InHours / GetDayLength());
+		const float DayCycleSeconds = PlayerDurationSeconds * DayLengthRatio;
 
 		// Update the playback position of the sequence.
 		FDaySequencePlaybackParams PlaybackParams;
@@ -1580,8 +1467,7 @@ FMovieSceneSequencePlaybackSettings ADaySequenceActor::GetPlaybackSettings(const
 
 	if (Sequence && Sequence->GetMovieScene())
 	{
-		// Convert the initial day time to sequence time
-		const UMovieScene* MovieScene = Sequence->GetMovieScene();
+		Settings.PlayRate = GetBasePlayRate();
 
 #if WITH_EDITORONLY_DATA
 		const float InitialHours = bOverrideInitialTimeOfDay ? GetTimeOfDayPreview() : GetInitialTimeOfDay();
@@ -1589,9 +1475,9 @@ FMovieSceneSequencePlaybackSettings ADaySequenceActor::GetPlaybackSettings(const
 		const float InitialHours = GetInitialTimeOfDay();
 #endif
 		
-		const float DayLengthHours = GetDayLength();
-		const float StartRatio = FMath::Frac(InitialHours / DayLengthHours);
-		const float DayCycleSeconds = TimePerCycle.ToSeconds() * StartRatio;
+		const float StartRatio = FMath::Frac(InitialHours / GetDayLength());
+		const float PlayerDurationSeconds = GetBaseDuration() * FDaySequenceTime::SecondsPerHour;
+		const float DayCycleSeconds = PlayerDurationSeconds * StartRatio;
 		Settings.StartTime = DayCycleSeconds;
 	}
 	return Settings;
diff --git a/Engine/Plugins/Experimental/DaySequence/Source/DaySequence/Public/DaySequenceActor.h b/Engine/Plugins/Experimental/DaySequence/Source/DaySequence/Public/DaySequenceActor.h
--- a/Engine/Plugins/Experimental/DaySequence/Source/DaySequence/Public/DaySequenceActor.h
+++ b/Engine/Plugins/Experimental/DaySequence/Source/DaySequence/Public/DaySequenceActor.h
@@ -424,6 +424,13 @@ protected:
 	/** Initialize SequencePlayer with transient root sequence */
 	void InitializePlayer();
 	void InitializeRootSequence();
+	void SetRootSequencePlaybackRange();
+
+	// Returns the rate required for the player to have a duration of TimePerCycle
+	float GetBasePlayRate() const;
+
+	// Returns the duration in hours of the player if its rate is 1
+	float GetBaseDuration() const;
 
 	/* Internal getter that validates the player. */
 	UDaySequencePlayer* GetSequencePlayerInternal() const;
@@ -453,9 +460,6 @@ protected:
 
 	void OnSequencePlayerUpdate(const UDaySequencePlayer& Player, FFrameTime CurrentTime, FFrameTime PreviousTime);
 	virtual void SequencePlayerUpdated(float CurrentTime, float PreviousTime);
-
-	/* Callback for changes to the DaySequence.UsePlayRate cvar. Keeps actor in a good state when the CVar changes. */
-	void OnUsePlayRateChanged(IConsoleVariable* UsePlayRateVariable);
 	
 public:
 	/**
@@ -614,9 +618,5 @@ protected:
 	
 	UPROPERTY(Transient)
 	TMap<TSubclassOf<UDaySequenceConditionTag>, TObjectPtr<UDaySequenceConditionTag>> TrackConditionMap;
-
-private:
-	/* Used to change TimePerCycle while a player is active. Must be called on server and all clients in order to prevent synchronization issues. */
-	void Runtime_SetTimePerCycle_Impl(float InHours);
 };
 
diff --git a/Engine/Plugins/Experimental/DaySequence/Source/DaySequenceEditor/Private/DaySequenceActorPreview.cpp b/Engine/Plugins/Experimental/DaySequence/Source/DaySequenceEditor/Private/DaySequenceActorPreview.cpp
--- a/Engine/Plugins/Experimental/DaySequence/Source/DaySequenceEditor/Private/DaySequenceActorPreview.cpp
+++ b/Engine/Plugins/Experimental/DaySequence/Source/DaySequenceEditor/Private/DaySequenceActorPreview.cpp
@@ -80,9 +80,16 @@ float FDaySequenceActorPreview::GetPreviewTime() const
 	float PreviewTime = 0.0f;
 	if (DaySequenceActor.IsValid() && DaySequencePreviewToolkit.IsValid())
 	{
+		// Convert sequencer time to equivalent game time.
 		const TSharedPtr<ISequencer> Sequencer = DaySequencePreviewToolkit.Pin()->GetSequencer();
-		const float SequencerTime = Sequencer->GetGlobalTime().AsSeconds() / 3600.f;
-		PreviewTime = SequencerTime / DaySequenceActor->GetTimePerCycle() * DaySequenceActor->GetDayLength();
+
+		const FFrameNumber LowerBound = Sequencer->GetPlaybackRange().GetLowerBoundValue();
+		const FFrameNumber UpperBound = Sequencer->GetPlaybackRange().GetUpperBoundValue();
+		const int32 Range = UpperBound.Value - LowerBound.Value;
+		const int32 CurrentTimeOffset = Sequencer->GetGlobalTime().Time.FrameNumber.Value - LowerBound.Value;
+		const float NormalizedTime = static_cast<float>(CurrentTimeOffset) / Range;
+
+		PreviewTime = NormalizedTime * DaySequenceActor->GetDayLength();
 	}
 	return PreviewTime;
 }
@@ -105,15 +112,17 @@ void FDaySequenceActorPreview::SetPreviewTime(float NewPreviewTime)
 	}
 	
 	const TSharedPtr<ISequencer> Sequencer = DaySequencePreviewToolkit.Pin()->GetSequencer();
-	const UMovieSceneSequence* MovieSequence = Sequencer->GetRootMovieSceneSequence();
-	const UMovieScene* MovieScene = MovieSequence->GetMovieScene();
-	const FFrameRate TickResolution = MovieScene->GetTickResolution();
 	
 	// Given game time (NewPreviewTime), convert to equivalent sequencer time.
-	const float SequencerTime = NewPreviewTime * DaySequenceActor->GetTimePerCycle() / DaySequenceActor->GetDayLength();
-	const float SequencerSeconds = SequencerTime * 3600.f;
-	const FFrameTime FrameRangeUpperBound = Sequencer->GetRootMovieSceneSequence()->GetMovieScene()->GetPlaybackRange().GetUpperBoundValue() - 1;
-	const FFrameTime SequencerFrameTime = FMath::Clamp<FFrameTime>(TickResolution.AsFrameTime(SequencerSeconds), 0, FrameRangeUpperBound);
+	const float NormalizedTime = NewPreviewTime / DaySequenceActor->GetDayLength();
+	
+	const FFrameNumber LowerBound = Sequencer->GetPlaybackRange().GetLowerBoundValue();
+	const FFrameNumber UpperBound = Sequencer->GetPlaybackRange().GetUpperBoundValue();
+	const int32 Range = UpperBound.Value - LowerBound.Value;
+	const int32 CurrentTimeOffset = NormalizedTime * Range;
+	const FFrameNumber CurrentTime = LowerBound + CurrentTimeOffset;
+
+	const FFrameTime SequencerFrameTime = FMath::Clamp<FFrameTime>(CurrentTime, 0, UpperBound);
 	
 	if (Sequencer->IsEvaluating())
 	{
diff --git a/Engine/Plugins/Experimental/DaySequence/Source/DaySequenceEditor/Private/DaySequenceEditorToolkit.cpp b/Engine/Plugins/Experimental/DaySequence/Source/DaySequenceEditor/Private/DaySequenceEditorToolkit.cpp
--- a/Engine/Plugins/Experimental/DaySequence/Source/DaySequenceEditor/Private/DaySequenceEditorToolkit.cpp
+++ b/Engine/Plugins/Experimental/DaySequence/Source/DaySequenceEditor/Private/DaySequenceEditorToolkit.cpp
@@ -301,6 +301,15 @@ void FDaySequenceEditorToolkit::InitializeInternal(
 	Sequencer = FModuleManager::LoadModuleChecked<ISequencerModule>("Sequencer").CreateSequencer(SequencerInitParams);
 	SpawnRegister->SetSequencer(Sequencer);
 	Sequencer->OnActorAddedToSequencer().AddSP(this, &FDaySequenceEditorToolkit::HandleActorAddedToSequencer);
+
+	// Set appropriate default playback speed based on the ratio of the root sequence duration to time per cycle.
+	if (const ADaySequenceActor* DayActor = RootActor ? RootActor : PreviewActor)
+	{
+		const UMovieScene* MovieScene = SequencerInitParams.RootSequence->GetMovieScene();
+		const float SequenceDurationHours = MovieScene->GetTickResolution().AsSeconds(MovieScene->GetPlaybackRange().GetUpperBoundValue()) / 3600.f;
+		const float DesiredDurationHours = DayActor->GetTimeOfDay();
+		Sequencer->SetPlaybackSpeed(SequenceDurationHours / DesiredDurationHours);
+	}
 	
 	const bool bIsSequenceEditor = !IsActorPreview();
 
@@ -775,9 +784,15 @@ void FDaySequenceEditorToolkit::OnGlobalTimeChanged()
 	// If the sequence we are viewing is the PreviewActor, also propagate sequencer time to actor.
 	if (ADaySequenceActor* DayActor = RootActor ? RootActor : PreviewActor)
 	{
-		// Given sequencer time (SequencerTimeHours), convert to equivalent game time (GameTimeHours).
-		const float SequencerTimeHours = Sequencer->GetGlobalTime().AsSeconds() / 3600.f;
-		const float GameTimeHours = SequencerTimeHours / DayActor->GetTimePerCycle() * DayActor->GetDayLength();
+		// Convert sequencer time to equivalent game time.
+		
+		const FFrameNumber LowerBound = Sequencer->GetPlaybackRange().GetLowerBoundValue();
+		const FFrameNumber UpperBound = Sequencer->GetPlaybackRange().GetUpperBoundValue();
+		const int32 Range = UpperBound.Value - LowerBound.Value;
+		const int32 CurrentTimeOffset = Sequencer->GetGlobalTime().Time.FrameNumber.Value - LowerBound.Value;
+		const float NormalizedTime = static_cast<float>(CurrentTimeOffset) / Range;
+		
+		const float GameTimeHours = NormalizedTime * DayActor->GetDayLength();
 		DayActor->ConditionalSetTimeOfDayPreview(GameTimeHours);
 	}
 }
