DaySequenceModifierComponent: Add a Time based blending mode.

BlendAmount now determines how quickly the modifier fades in/out when the blend target enters/exits the volume.

Can yield unintuitive yet correct results when using Maximum or Override blend policies.

diff --git a/Engine/Plugins/Experimental/DaySequence/Source/DaySequence/Private/DaySequenceModifierComponent.cpp b/Engine/Plugins/Experimental/DaySequence/Source/DaySequence/Private/DaySequenceModifierComponent.cpp
--- a/Engine/Plugins/Experimental/DaySequence/Source/DaySequence/Private/DaySequenceModifierComponent.cpp
+++ b/Engine/Plugins/Experimental/DaySequence/Source/DaySequence/Private/DaySequenceModifierComponent.cpp
@@ -1095,6 +1095,10 @@ float UDaySequenceModifierComponent::GetBlendWeight() const
 float UDaySequenceModifierComponent::UpdateInternalBlendWeight()
 {
 	CSV_SCOPED_TIMING_STAT(DaySequence, UpdateInternalBlendWeight);
+
+	const float CurrentTime = GetWorld()->GetTimeSeconds();
+	const float DeltaTime = CurrentTime - InternalBlendWeightUpdatedTime;
+	const float PreviousBlendWeight = InternalBlendWeight;
 	
 	switch (Mode)
 	{
@@ -1102,7 +1106,9 @@ float UDaySequenceModifierComponent::UpdateInternalBlendWeight()
 	case EDaySequenceModifierMode::Global:
 		InternalBlendWeight = 1.f;
 		break;
-	
+
+	case EDaySequenceModifierMode::Time:
+		// Intentional fallthrough in order to determine if we are currently inside or outside the volume
 	case EDaySequenceModifierMode::Volume:
 		if (FVector BlendPosition; GetBlendPosition(BlendPosition))
 		{
@@ -1126,6 +1132,18 @@ float UDaySequenceModifierComponent::UpdateInternalBlendWeight()
 		break;
 	}
 
+	if (Mode == EDaySequenceModifierMode::Time)
+	{
+		if (InternalBlendWeight > 0.f)
+		{
+			InternalBlendWeight = FMath::Clamp(PreviousBlendWeight + DeltaTime / FMath::Max(BlendTime, UE_SMALL_NUMBER), UE_SMALL_NUMBER, 1.f);
+		}
+		else
+		{
+			InternalBlendWeight = FMath::Clamp(PreviousBlendWeight - DeltaTime / FMath::Max(BlendTime, UE_SMALL_NUMBER), 0.f, 1.f);
+		}
+	}
+	
 	if (UE::DaySequence::CVarModifierDisableWhenInvisible.GetValueOnAnyThread())
 	{
 		if (!IsVisible())
@@ -1134,6 +1152,7 @@ float UDaySequenceModifierComponent::UpdateInternalBlendWeight()
 		}
 	}
 
+	InternalBlendWeightUpdatedTime = CurrentTime;
 	return InternalBlendWeight;
 }
 
diff --git a/Engine/Plugins/Experimental/DaySequence/Source/DaySequence/Public/DaySequenceModifierComponent.h b/Engine/Plugins/Experimental/DaySequence/Source/DaySequence/Public/DaySequenceModifierComponent.h
--- a/Engine/Plugins/Experimental/DaySequence/Source/DaySequence/Public/DaySequenceModifierComponent.h
+++ b/Engine/Plugins/Experimental/DaySequence/Source/DaySequence/Public/DaySequenceModifierComponent.h
@@ -54,8 +54,11 @@ enum class EDaySequenceModifierMode : uint8
 	// Blend weight is always 1.0.
 	Global,
 
-	// Blend weight smoothly moves between 0.0 and 1.0 as the blend target crosses the volume boundary.
-	Volume
+	// Blend weight smoothly moves between 0.0 and 1.0 according to how far the blend target is from the volume boundary.
+	Volume,
+
+	// Blend weight smoothly moves between 0.0 and 1.0 at a fixed rate according to when the blend target last crossed the volume boundary.
+	Time
 };
 
 /** Enum specifying how the modifier resolves the user specified blend weight against the internal blend weight. */
@@ -284,6 +287,10 @@ protected:
 	UPROPERTY(EditAnywhere, BlueprintReadWrite, Category="Day Sequence", meta=(DisplayAfter="Mode", EditCondition="Mode==EDaySequenceModifierMode::Volume", EditConditionHides))
 	float BlendAmount;
 
+	/** Defines the amount of time (in seconds) that it takes for blend weight to move across the full range (0.0, 1.0) when Mode == EDaySequenceModifierMode::Time. */
+	UPROPERTY(EditAnywhere, BlueprintReadWrite, Category="Day Sequence", meta=(DisplayAfter="Mode", EditCondition="Mode==EDaySequenceModifierMode::Time", EditConditionHides))
+	float BlendTime;
+
 	/** User specified blend weight. The final blend weight is determined by BlendPolicy. */
 	UPROPERTY(EditAnywhere, BlueprintReadWrite, Setter, Getter, Category="Day Sequence", meta=(DisplayAfter="BlendPolicy", EditCondition="BlendPolicy!=EDaySequenceModifierUserBlendPolicy::Ignored", EditConditionHides))
 	float UserBlendWeight;
@@ -362,6 +369,9 @@ private:
 	 * Used to compute effective blend weight along with BlendPolicy and UserBlendWeight.
 	 */
 	float InternalBlendWeight;
+
+	/** Used by UpdateInternalBlendWeight to compute time delta between InternalBlendWeight updates. */
+	float InternalBlendWeightUpdatedTime;
 	
 	/**
 	 * Creates and adds or marks for preserve all subsections that the modifier is responsible for.
