DaySequenceModifierComponent: Fix crash from unchecked GetWorld dereference.

diff --git a/Engine/Plugins/Experimental/DaySequence/Source/DaySequence/Private/DaySequenceModifierComponent.cpp b/Engine/Plugins/Experimental/DaySequence/Source/DaySequence/Private/DaySequenceModifierComponent.cpp
--- a/Engine/Plugins/Experimental/DaySequence/Source/DaySequence/Private/DaySequenceModifierComponent.cpp
+++ b/Engine/Plugins/Experimental/DaySequence/Source/DaySequence/Private/DaySequenceModifierComponent.cpp
@@ -580,7 +580,7 @@ void UDaySequenceModifierComponent::SetInitialTimeOfDay()
 				{
 					if (WantsStaticTime())
 					{
-						OutRequest.BlendWeight = Mode == EDaySequenceModifierMode::Volume ? GetBlendWeight() : 1.f;
+						OutRequest.BlendWeight = GetBlendWeight();
 						OutRequest.StaticTime = DayNightCycle == EDayNightCycleMode::RandomFixedTime ? RandomTime : DayNightCycleTime;
 						return true;
 					}
@@ -1096,8 +1096,6 @@ float UDaySequenceModifierComponent::UpdateInternalBlendWeight()
 {
 	CSV_SCOPED_TIMING_STAT(DaySequence, UpdateInternalBlendWeight);
 
-	const float CurrentTime = GetWorld()->GetTimeSeconds();
-	const float DeltaTime = CurrentTime - InternalBlendWeightUpdatedTime;
 	const float PreviousBlendWeight = InternalBlendWeight;
 	
 	switch (Mode)
@@ -1134,13 +1132,24 @@ float UDaySequenceModifierComponent::UpdateInternalBlendWeight()
 
 	if (Mode == EDaySequenceModifierMode::Time)
 	{
-		if (InternalBlendWeight > 0.f)
+		if (const UWorld* World = GetWorld())
 		{
-			InternalBlendWeight = FMath::Clamp(PreviousBlendWeight + DeltaTime / FMath::Max(BlendTime, UE_SMALL_NUMBER), UE_SMALL_NUMBER, 1.f);
+			const float CurrentTime = World->GetTimeSeconds();
+            const float DeltaTime = CurrentTime - TimedBlendingLastUpdated;
+        	TimedBlendingLastUpdated = CurrentTime;
+            
+        	if (InternalBlendWeight > 0.f)
+        	{
+        		InternalBlendWeight = FMath::Clamp(PreviousBlendWeight + DeltaTime / FMath::Max(BlendTime, UE_SMALL_NUMBER), UE_SMALL_NUMBER, 1.f);
+        	}
+        	else
+        	{
+        		InternalBlendWeight = FMath::Clamp(PreviousBlendWeight - DeltaTime / FMath::Max(BlendTime, UE_SMALL_NUMBER), 0.f, 1.f);
+        	}
 		}
 		else
 		{
-			InternalBlendWeight = FMath::Clamp(PreviousBlendWeight - DeltaTime / FMath::Max(BlendTime, UE_SMALL_NUMBER), 0.f, 1.f);
+			InternalBlendWeight = PreviousBlendWeight;
 		}
 	}
 	
@@ -1152,7 +1161,6 @@ float UDaySequenceModifierComponent::UpdateInternalBlendWeight()
 		}
 	}
 
-	InternalBlendWeightUpdatedTime = CurrentTime;
 	return InternalBlendWeight;
 }
 
diff --git a/Engine/Plugins/Experimental/DaySequence/Source/DaySequence/Public/DaySequenceModifierComponent.h b/Engine/Plugins/Experimental/DaySequence/Source/DaySequence/Public/DaySequenceModifierComponent.h
--- a/Engine/Plugins/Experimental/DaySequence/Source/DaySequence/Public/DaySequenceModifierComponent.h
+++ b/Engine/Plugins/Experimental/DaySequence/Source/DaySequence/Public/DaySequenceModifierComponent.h
@@ -371,7 +371,7 @@ private:
 	float InternalBlendWeight;
 
 	/** Used by UpdateInternalBlendWeight to compute time delta between InternalBlendWeight updates. */
-	float InternalBlendWeightUpdatedTime;
+	float TimedBlendingLastUpdated;
 	
 	/**
 	 * Creates and adds or marks for preserve all subsections that the modifier is responsible for.
