DaySequence: Added StaticTimeMode property to the DaySequenceActor.

StaticTimeMode supports two modes currently: RootTimeWarp and LocalTimeWarp. RootTimeWarp behaves as it does today, where StaticTime will directly influence Root Sequence evaluation time. In contrast, LocalTimeWarp is a contract implemented by DaySequenceModifierComponents where the time warp only applies to the modifier's sequences so that those sequences are always evaluated at the frame that matches the desired fixed time. The benefit of LocalTimeWarp is that the DaySequenceActor will use hierarchical blending to lerp towards the look of the target time rather than lerping actual time which can cause jarring visual shifts as it might progress through the sunrist/sunset shifts to achieve the target look.

By using TimeWarp to achieve this behavior, users can reuse existing sequences without needing to reauthor a sequence specifically for the desired time of day.

diff --git a/Engine/Plugins/Experimental/DaySequence/Source/DaySequence/Private/DaySequenceActor.cpp b/Engine/Plugins/Experimental/DaySequence/Source/DaySequence/Private/DaySequenceActor.cpp
--- a/Engine/Plugins/Experimental/DaySequence/Source/DaySequence/Private/DaySequenceActor.cpp
+++ b/Engine/Plugins/Experimental/DaySequence/Source/DaySequence/Private/DaySequenceActor.cpp
@@ -1122,7 +1122,7 @@ void ADaySequenceActor::WarpEvaluationRange(FMovieSceneEvaluationRange& InOutRan
 	};
 
 	// Warp with static time if necessary
-	if (HasStaticTimeOfDay())
+	if (HasStaticTimeOfDay() && StaticTimeMode == EDaySequenceStaticTimeMode::RootTimeWarp)
 	{
 		const float StaticTimeInGameHours = GetStaticTimeOfDay();
 
@@ -1312,6 +1312,16 @@ float ADaySequenceActor::GetStaticTimeOfDay() const
 	return std::numeric_limits<float>::lowest();
 }
 
+void ADaySequenceActor::SetStaticTimeMode(EDaySequenceStaticTimeMode Mode)
+{
+	StaticTimeMode = Mode;
+}
+
+EDaySequenceStaticTimeMode ADaySequenceActor::GetStaticTimeMode() const
+{
+	return StaticTimeMode;
+}
+
 void ADaySequenceActor::RegisterStaticTimeContributor(const UE::DaySequence::FStaticTimeContributor& NewContributor) const
 {
 	if (!StaticTimeManager)
diff --git a/Engine/Plugins/Experimental/DaySequence/Source/DaySequence/Private/DaySequenceModifierComponent.cpp b/Engine/Plugins/Experimental/DaySequence/Source/DaySequence/Private/DaySequenceModifierComponent.cpp
--- a/Engine/Plugins/Experimental/DaySequence/Source/DaySequence/Private/DaySequenceModifierComponent.cpp
+++ b/Engine/Plugins/Experimental/DaySequence/Source/DaySequence/Private/DaySequenceModifierComponent.cpp
@@ -23,6 +23,8 @@
 #include "ProfilingDebugging/CsvProfiler.h"
 #include "Sections/MovieSceneSubSection.h"
 #include "Tracks/MovieSceneSubTrack.h"
+#include "Variants/MovieSceneTimeWarpVariant.h"
+#include "Variants/MovieSceneTimeWarpVariantPayloads.h"
 
 #define LOCTEXT_NAMESPACE "DaySequenceModifierComponent"
 
@@ -754,11 +756,31 @@ UMovieSceneSubSection* UDaySequenceModifierComponent::InitializeDaySequence(cons
 												| (bIgnoreBias ? EMovieSceneSubSectionFlags::IgnoreHierarchicalBias : EMovieSceneSubSectionFlags::None)
 												| (bBlendHierarchicalBias ? EMovieSceneSubSectionFlags::BlendHierarchicalBias : EMovieSceneSubSectionFlags::None);
 
+		const TRange<FFrameNumber> PlaybackRange = MovieScene->GetPlaybackRange();
 		SubSection->SetSequence(Sequence);
-		SubSection->SetRange(MovieScene->GetPlaybackRange());
+		SubSection->SetRange(PlaybackRange);
 		SubSection->SetIsActive(bActivate);
-		
-		TargetActor->UpdateSubSectionTimeScale(SubSection);
+
+		if (TargetActor->GetStaticTimeMode() == EDaySequenceStaticTimeMode::LocalTimeWarp)
+		{
+			if (FMovieSceneTimeWarpVariant* Variant = SubSection->GetTimeWarp())
+			{
+				UMovieScene* SubMovieScene = Sequence->GetMovieScene();
+				const FFrameRate TickResolution = SubMovieScene->GetTickResolution();
+				const FQualifiedFrameTime SubDuration = FQualifiedFrameTime(
+					UE::MovieScene::DiscreteSize(SubMovieScene->GetPlaybackRange()),
+					TickResolution);
+
+				const float TargetFixedRatio = FMath::Fmod(DayNightCycleTime / TargetActor->GetDayLength(), 1.0f);
+				const FTimecode TargetFixedTimecode = FTimecode(TargetFixedRatio * SubDuration.AsSeconds(), TickResolution, false, false);
+				const FFrameNumber TargetFrame = TargetFixedTimecode.ToFrameNumber(TickResolution);
+				Variant->Set(FMovieSceneTimeWarpFixedFrame(TargetFrame));
+			}
+		}
+		else
+		{
+			TargetActor->UpdateSubSectionTimeScale(SubSection);
+		}
 
 		RootTrack->AddSection(*SubSection);
 
@@ -767,7 +789,7 @@ UMovieSceneSubSection* UDaySequenceModifierComponent::InitializeDaySequence(cons
 		// width to facilitate blending, the label is clipped. So we use EaseOut here instead and
 		// ensure that the weight is inverted in Evaluate().
 		SubSection->Easing.bManualEaseOut = true;
-		SubSection->Easing.ManualEaseOutDuration = MovieScene->GetPlaybackRange().Size<FFrameNumber>().Value;
+		SubSection->Easing.ManualEaseOutDuration = PlaybackRange.Size<FFrameNumber>().Value;
 		EasingFunction->Initialize(UDaySequenceModifierEasingFunction::EEasingFunctionType::EaseOut);
 		SubSection->Easing.EaseOut = EasingFunction;
 
diff --git a/Engine/Plugins/Experimental/DaySequence/Source/DaySequence/Public/DaySequenceActor.h b/Engine/Plugins/Experimental/DaySequence/Source/DaySequence/Public/DaySequenceActor.h
--- a/Engine/Plugins/Experimental/DaySequence/Source/DaySequence/Public/DaySequenceActor.h
+++ b/Engine/Plugins/Experimental/DaySequence/Source/DaySequence/Public/DaySequenceActor.h
@@ -76,6 +76,16 @@ namespace UE::DaySequence
 
 DECLARE_DYNAMIC_MULTICAST_DELEGATE_OneParam(FOnTimeOfDayPreviewChanged, float, Time);
 
+UENUM()
+enum class EDaySequenceStaticTimeMode
+{
+	// Static time will apply a time warp on the RootSequence/global evaluation time.
+	RootTimeWarp,
+	// Static time will apply a time warp local to the modifier's sequence evaluation time.
+	// The local time warp will still influence the apparent time, but not the global time.
+	LocalTimeWarp
+};
+
 UCLASS(Blueprintable)
 class DAYSEQUENCE_API ADaySequenceActor
 	: public AInfo
@@ -291,6 +301,9 @@ public:
 	UFUNCTION(BlueprintCallable, Category="TimeOfDay")
 	float GetStaticTimeOfDay() const;
 
+	void SetStaticTimeMode(EDaySequenceStaticTimeMode Mode);
+	EDaySequenceStaticTimeMode GetStaticTimeMode() const;
+
 	void RegisterStaticTimeContributor(const UE::DaySequence::FStaticTimeContributor& NewContributor) const;
 	void UnregisterStaticTimeContributor(const UObject* InUserObject) const;
 	
@@ -562,6 +575,9 @@ protected:
 	
 	TSharedPtr<UE::DaySequence::FStaticTimeManager> StaticTimeManager;
 
+	UPROPERTY(EditAnywhere, Category = RuntimeDayCycle)
+	EDaySequenceStaticTimeMode StaticTimeMode = EDaySequenceStaticTimeMode::RootTimeWarp;
+
 	UPROPERTY(Transient)
 	TObjectPtr<UDaySequenceCameraModifierManager> CameraModifierManager;
 	
