DaySequence: Deprecated DaySequenceCollection in favor of an array of DaySequenceCollections on both the DaySequenceActor and DaySequenceModifierComponent.

The addition of an array of collections will allow for greater flexibility in modularizing sequences particularly with the recent introduction of local time warp for Modifier Components, where it is useful to reuse a particular combination of DaySequenceCollections to recreate the visual target of a modifier at a specific time of day.

diff --git a/Engine/Plugins/Experimental/DaySequence/Source/DaySequence/Private/Actors/SunMoonDaySequenceActor.cpp b/Engine/Plugins/Experimental/DaySequence/Source/DaySequence/Private/Actors/SunMoonDaySequenceActor.cpp
--- a/Engine/Plugins/Experimental/DaySequence/Source/DaySequence/Private/Actors/SunMoonDaySequenceActor.cpp
+++ b/Engine/Plugins/Experimental/DaySequence/Source/DaySequence/Private/Actors/SunMoonDaySequenceActor.cpp
@@ -34,6 +34,6 @@ ASunMoonDaySequenceActor::ASunMoonDaySequenceActor(const FObjectInitializer& Ini
 	{
 		// Override the default collection (which animates the moon and sky material)
         static ConstructorHelpers::FObjectFinder<UDaySequenceCollectionAsset> DefaultCollection(TEXT("/DaySequence/DSCA_24hr.DSCA_24hr"));
-        DaySequenceCollection = DefaultCollection.Object.Get();
+        DaySequenceCollections.Add(DefaultCollection.Object.Get());
 	}
 }
\ No newline at end of file
diff --git a/Engine/Plugins/Experimental/DaySequence/Source/DaySequence/Private/DaySequenceActor.cpp b/Engine/Plugins/Experimental/DaySequence/Source/DaySequence/Private/DaySequenceActor.cpp
--- a/Engine/Plugins/Experimental/DaySequence/Source/DaySequence/Private/DaySequenceActor.cpp
+++ b/Engine/Plugins/Experimental/DaySequence/Source/DaySequence/Private/DaySequenceActor.cpp
@@ -22,6 +22,7 @@
 #include "Sections/MovieSceneSubSection.h"
 #include "Tracks/MovieSceneSubTrack.h"
 #include "UObject/ConstructorHelpers.h"
+#include "UObject/Package.h"
 #include "UObject/UObjectThreadContext.h" // For FUObjectThreadContext
 #include "ProfilingDebugging/CsvProfiler.h"
 
@@ -222,13 +223,16 @@ UDaySequencePlayer* ADaySequenceActor::GetSequencePlayerInternal() const
 
 bool ADaySequenceActor::ContainsDaySequence(const UDaySequence* InDaySequence)
 {
-	if (InDaySequence && DaySequenceCollection)
+	if (InDaySequence)
 	{
-		for (const FDaySequenceCollectionEntry& SequenceAsset : DaySequenceCollection->DaySequences)
+		for (UDaySequenceCollectionAsset* Collection : DaySequenceCollections)
 		{
-			if (SequenceAsset.Sequence == InDaySequence)
+			for (const FDaySequenceCollectionEntry& SequenceAsset : Collection->DaySequences)
 			{
-				return true;
+				if (SequenceAsset.Sequence == InDaySequence)
+				{
+					return true;
+				}
 			}
 		}
 	}
@@ -293,6 +297,15 @@ void ADaySequenceActor::PostLoad()
 		LocalSpriteComponent->SetupAttachment(RootComponent);
 	}
 #endif
+	
+PRAGMA_DISABLE_DEPRECATION_WARNINGS
+	// Transfer DaySequenceCollection (deprecated) to DaySequenceCollections
+	if (DaySequenceCollection)
+	{
+		DaySequenceCollections.Add(DaySequenceCollection);
+		DaySequenceCollection = nullptr;
+	}
+PRAGMA_ENABLE_DEPRECATION_WARNINGS
 }
 
 void ADaySequenceActor::BeginDestroy()
@@ -430,9 +443,9 @@ UMovieSceneSequence* ADaySequenceActor::RetrieveOwnedSequence() const
 
 bool ADaySequenceActor::GetReferencedContentObjects(TArray<UObject*>& Objects) const
 {
-	if (DaySequenceCollection)
+	for (UDaySequenceCollectionAsset* Collection : DaySequenceCollections)
 	{
-		Objects.Add(DaySequenceCollection);
+		Objects.Add(Collection);
 	}
 	Super::GetReferencedContentObjects(Objects);
 	return true;
@@ -472,7 +485,7 @@ void ADaySequenceActor::PostEditChangeProperty(FPropertyChangedEvent& PropertyCh
 	{
 		SetInitialTimeOfDay(GetInitialTimeOfDay());
 	}
-	else if (PropertyName == GET_MEMBER_NAME_CHECKED(ADaySequenceActor, DaySequenceCollection))
+	else if (PropertyName == GET_MEMBER_NAME_CHECKED(ADaySequenceActor, DaySequenceCollections))
 	{
 		ReinitializeActor();
 	}
@@ -770,14 +783,14 @@ void ADaySequenceActor::InitializeDaySequences()
 {
 	SubSections.Empty();
 	
-	if (DaySequenceCollection)
+	for (UDaySequenceCollectionAsset* Collection : DaySequenceCollections)
 	{
-		for (const FDaySequenceCollectionEntry& Entry : DaySequenceCollection->DaySequences)
+		for (const FDaySequenceCollectionEntry& Entry : Collection->DaySequences)
 		{
 			InitializeDaySequence(Entry);
 		}
 
-		for (TInstancedStruct<FProceduralDaySequence>& ProceduralDaySequence : DaySequenceCollection->ProceduralDaySequences)
+		for (TInstancedStruct<FProceduralDaySequence>& ProceduralDaySequence : Collection->ProceduralDaySequences)
 		{
 			if (!ProceduralDaySequence.IsValid())
 			{
@@ -1427,7 +1440,7 @@ void ADaySequenceActor::UpdateRootSequence()
 	// Set up a callback for when UObjects are modified so we can catch changes to our 
 	FCoreUObjectDelegates::OnObjectModified.AddWeakLambda(this, [this](UObject* InObject)
 	{
-		if (IsValid(InObject) && DaySequenceCollection == InObject)
+		if (IsValid(InObject) && InObject->IsA(UDaySequenceCollectionAsset::StaticClass()) && DaySequenceCollections.Contains(InObject))
 		{
 			// We update on next tick because calling UpdateScalabilitySequence here
 			// is too early (our collection doesn't have the changes that triggered this invocation).
diff --git a/Engine/Plugins/Experimental/DaySequence/Source/DaySequence/Private/DaySequenceModifierComponent.cpp b/Engine/Plugins/Experimental/DaySequence/Source/DaySequence/Private/DaySequenceModifierComponent.cpp
--- a/Engine/Plugins/Experimental/DaySequence/Source/DaySequence/Private/DaySequenceModifierComponent.cpp
+++ b/Engine/Plugins/Experimental/DaySequence/Source/DaySequence/Private/DaySequenceModifierComponent.cpp
@@ -314,6 +314,20 @@ void UDaySequenceModifierComponent::DaySequenceUpdate()
 	}
 }
 
+void UDaySequenceModifierComponent::PostLoad()
+{
+	Super::PostLoad();
+
+PRAGMA_DISABLE_DEPRECATION_WARNINGS
+	// Transfer DaySequenceCollection (deprecated) to DaySequenceCollections
+	if (DaySequenceCollection)
+	{
+		DaySequenceCollections.Add(DaySequenceCollection);
+		DaySequenceCollection = nullptr;
+	}
+PRAGMA_ENABLE_DEPRECATION_WARNINGS
+}
+
 void UDaySequenceModifierComponent::BeginPlay()
 {
 	Super::BeginPlay();
@@ -670,14 +684,14 @@ void UDaySequenceModifierComponent::ReinitializeSubSequence(ADaySequenceActor::F
 		
 		if (bUseCollection)
 		{
-			if (DaySequenceCollection)
+			for (UDaySequenceCollectionAsset* Collection : DaySequenceCollections)
 			{
-				for (const FDaySequenceCollectionEntry& Entry : DaySequenceCollection->DaySequences)
+				for (const FDaySequenceCollectionEntry& Entry : Collection->DaySequences)
 				{
 					InitializeDaySequence(Entry);
 				}
 
-				for (TInstancedStruct<FProceduralDaySequence>& ProceduralEntry : DaySequenceCollection->ProceduralDaySequences)
+				for (TInstancedStruct<FProceduralDaySequence>& ProceduralEntry : Collection->ProceduralDaySequences)
 				{
 					if (!ProceduralEntry.IsValid())
 					{
diff --git a/Engine/Plugins/Experimental/DaySequence/Source/DaySequence/Public/DaySequenceActor.h b/Engine/Plugins/Experimental/DaySequence/Source/DaySequence/Public/DaySequenceActor.h
--- a/Engine/Plugins/Experimental/DaySequence/Source/DaySequence/Public/DaySequenceActor.h
+++ b/Engine/Plugins/Experimental/DaySequence/Source/DaySequence/Public/DaySequenceActor.h
@@ -481,10 +481,14 @@ public:
 	 **/
 	UPROPERTY(EditAnywhere, replicated, BlueprintReadWrite, Category=RuntimeDayCycle, meta=(EditCondition=bUseInterpCurve, DisplayAfter=bRunDayCycle))
 	TObjectPtr<UCurveFloat> DayInterpCurve;
-	
-	UPROPERTY(EditAnywhere, Category="Sequence")
+
+	UE_DEPRECATED(5.6, "DaySequenceCollection is deprecated in favor of the array property. Please use DaySequenceCollections instead.")
+	UPROPERTY()
 	TObjectPtr<UDaySequenceCollectionAsset> DaySequenceCollection;
 
+	UPROPERTY(EditAnywhere, Category="Sequence")
+	TArray<TObjectPtr<UDaySequenceCollectionAsset>> DaySequenceCollections;
+
 	/** User-defined bias to apply to sequences in DaySequenceCollection. */
 	UPROPERTY(BlueprintReadWrite, EditAnywhere, Category="Sequence", DisplayName="Collection Bias", meta=(DisplayAfter="DaySequenceCollection"))
 	int32 Bias;
diff --git a/Engine/Plugins/Experimental/DaySequence/Source/DaySequence/Public/DaySequenceModifierComponent.h b/Engine/Plugins/Experimental/DaySequence/Source/DaySequence/Public/DaySequenceModifierComponent.h
--- a/Engine/Plugins/Experimental/DaySequence/Source/DaySequence/Public/DaySequenceModifierComponent.h
+++ b/Engine/Plugins/Experimental/DaySequence/Source/DaySequence/Public/DaySequenceModifierComponent.h
@@ -205,6 +205,10 @@ public:
 
 	void EmptyVolumeShapeComponents();
 	void AddVolumeShapeComponent(const FComponentReference& InShapeReference);
+
+	/*~ Begin UObject interface */
+	virtual void PostLoad() override;
+	/*~ End UObject interface */
 	
 	/*~ Begin UActorComponent interface */
 	virtual void BeginPlay() override;
@@ -269,8 +273,13 @@ protected:
     TObjectPtr<UDaySequence> TransientSequence;
 	
 	/** The user provided collection. This is an alternative to UserDaySequence. */
-	UPROPERTY(EditAnywhere, Category="Day Sequence", meta=(EditCondition="bUseCollection", EditConditionHides, DisplayAfter="bUseCollection"))
+	UE_DEPRECATED(5.6, "DaySequenceCollection is deprecated in favor of the array property. Please use DaySequenceCollections instead.")
+	UPROPERTY()
 	TObjectPtr<UDaySequenceCollectionAsset> DaySequenceCollection;
+	
+	/** The user provided collection. This is an alternative to UserDaySequence. */
+	UPROPERTY(EditAnywhere, Category="Day Sequence", meta=(EditCondition="bUseCollection", EditConditionHides, DisplayAfter="bUseCollection"))
+	TArray<TObjectPtr<UDaySequenceCollectionAsset>> DaySequenceCollections;
 
 	UPROPERTY(Instanced, Transient)
 	TObjectPtr<UDaySequenceModifierEasingFunction> EasingFunction;
diff --git a/Engine/Plugins/Experimental/DaySequence/Source/DaySequenceEditor/Private/DaySequenceActorFactory.cpp b/Engine/Plugins/Experimental/DaySequence/Source/DaySequenceEditor/Private/DaySequenceActorFactory.cpp
--- a/Engine/Plugins/Experimental/DaySequence/Source/DaySequenceEditor/Private/DaySequenceActorFactory.cpp
+++ b/Engine/Plugins/Experimental/DaySequence/Source/DaySequenceEditor/Private/DaySequenceActorFactory.cpp
@@ -44,7 +44,10 @@ UObject* UDaySequenceActorFactory::GetAssetFromActorInstance(AActor* Instance)
 {
 	if (const ADaySequenceActor* DaySequenceActor = Cast<ADaySequenceActor>(Instance))
 	{
-		return DaySequenceActor->DaySequenceCollection;
+		if (DaySequenceActor->DaySequenceCollections.Num() > 0)
+		{
+			return DaySequenceActor->DaySequenceCollections[0];
+		}
 	}
 	return nullptr;
 }
