DaySequence: Fixed crash on null collection entries.

diff --git a/Engine/Plugins/Experimental/DaySequence/Source/DaySequence/Private/DaySequenceActor.cpp b/Engine/Plugins/Experimental/DaySequence/Source/DaySequence/Private/DaySequenceActor.cpp
--- a/Engine/Plugins/Experimental/DaySequence/Source/DaySequence/Private/DaySequenceActor.cpp
+++ b/Engine/Plugins/Experimental/DaySequence/Source/DaySequence/Private/DaySequenceActor.cpp
@@ -227,6 +227,11 @@ bool ADaySequenceActor::ContainsDaySequence(const UDaySequence* InDaySequence)
 	{
 		for (UDaySequenceCollectionAsset* Collection : DaySequenceCollections)
 		{
+			if (!Collection)
+			{
+				continue;
+			}
+			
 			for (const FDaySequenceCollectionEntry& SequenceAsset : Collection->DaySequences)
 			{
 				if (SequenceAsset.Sequence == InDaySequence)
@@ -445,7 +450,10 @@ bool ADaySequenceActor::GetReferencedContentObjects(TArray<UObject*>& Objects) c
 {
 	for (UDaySequenceCollectionAsset* Collection : DaySequenceCollections)
 	{
-		Objects.Add(Collection);
+		if (Collection)
+		{
+			Objects.Add(Collection);
+		}
 	}
 	Super::GetReferencedContentObjects(Objects);
 	return true;
@@ -633,9 +641,12 @@ void ADaySequenceActor::InitializeRootSequence()
 		if (MarkDaySequences(&SectionsToPreserve))
 		{
 			InitializeDaySequences();
+			OnPostInitializeDaySequences.Broadcast(nullptr);
+		}
+		else
+		{
+			OnPostInitializeDaySequences.Broadcast(&SectionsToPreserve);
 		}
-
-		OnPostInitializeDaySequences.Broadcast(&SectionsToPreserve);
 
 		// Remove all unmarked sections
 		for (const TPair<UMovieSceneSubSection*, bool>& Section : SectionsToPreserve)
@@ -785,6 +796,11 @@ void ADaySequenceActor::InitializeDaySequences()
 	
 	for (UDaySequenceCollectionAsset* Collection : DaySequenceCollections)
 	{
+		if (!Collection)
+		{
+			continue;
+		}
+		
 		for (const FDaySequenceCollectionEntry& Entry : Collection->DaySequences)
 		{
 			InitializeDaySequence(Entry);
diff --git a/Engine/Plugins/Experimental/DaySequence/Source/DaySequence/Private/DaySequenceModifierComponent.cpp b/Engine/Plugins/Experimental/DaySequence/Source/DaySequence/Private/DaySequenceModifierComponent.cpp
--- a/Engine/Plugins/Experimental/DaySequence/Source/DaySequence/Private/DaySequenceModifierComponent.cpp
+++ b/Engine/Plugins/Experimental/DaySequence/Source/DaySequence/Private/DaySequenceModifierComponent.cpp
@@ -686,6 +686,11 @@ void UDaySequenceModifierComponent::ReinitializeSubSequence(ADaySequenceActor::F
 		{
 			for (UDaySequenceCollectionAsset* Collection : DaySequenceCollections)
 			{
+				if (!Collection)
+				{
+					continue;
+				}
+				
 				for (const FDaySequenceCollectionEntry& Entry : Collection->DaySequences)
 				{
 					InitializeDaySequence(Entry);
