DaySequence: Reinitialize actor on changes to StaticTimeMode.

diff --git a/Engine/Plugins/Experimental/DaySequence/Source/DaySequence/Private/DaySequenceActor.cpp b/Engine/Plugins/Experimental/DaySequence/Source/DaySequence/Private/DaySequenceActor.cpp
--- a/Engine/Plugins/Experimental/DaySequence/Source/DaySequence/Private/DaySequenceActor.cpp
+++ b/Engine/Plugins/Experimental/DaySequence/Source/DaySequence/Private/DaySequenceActor.cpp
@@ -497,6 +497,10 @@ void ADaySequenceActor::PostEditChangeProperty(FPropertyChangedEvent& PropertyCh
 	{
 		ReinitializeActor();
 	}
+	else if (PropertyName == GET_MEMBER_NAME_CHECKED(ADaySequenceActor, StaticTimeMode))
+	{
+		ReinitializeActor();
+	}
 	else if (PropertyChangedEvent.ChangeType == EPropertyChangeType::Unspecified)
 	{
 		// This handles undo/redo transactions.
