DaySequence: Fixed Refresh Day Sequence Actor not properly rebuilding/refreshing the root sequence.

diff --git a/Engine/Plugins/Experimental/DaySequence/Source/DaySequence/Private/DaySequenceActor.cpp b/Engine/Plugins/Experimental/DaySequence/Source/DaySequence/Private/DaySequenceActor.cpp
--- a/Engine/Plugins/Experimental/DaySequence/Source/DaySequence/Private/DaySequenceActor.cpp
+++ b/Engine/Plugins/Experimental/DaySequence/Source/DaySequence/Private/DaySequenceActor.cpp
@@ -1465,14 +1465,13 @@ void ADaySequenceActor::UpdateRootSequence()
 			// We update on next tick because calling UpdateScalabilitySequence here
 			// is too early (our collection doesn't have the changes that triggered this invocation).
 			bUpdateRootSequenceOnTick = true;
-
-			// Must be empty for a full reinit to take place (on next tick)
-			SubSections.Empty();
 		}
 	});
 #endif
 	
-	// For now, just rebuild the root sequence and GC the old one.
+	// Force rebuild the root sequence.
+	SubSections.Empty();
+	RootSequence = nullptr;
 	InitializeRootSequence();
 }
 
