DaySequence: Fixed regression where editing DaySequenceActor properties in a details panel with a root sequence open would reset the global sequence time.
Fixed an issue where restoring the preview time based on the root sequence global time was incorrect when focused on a sequence with a different playback range than the root sequence.

diff --git a/Engine/Plugins/Experimental/DaySequence/Source/DaySequence/Private/DaySequenceActor.cpp b/Engine/Plugins/Experimental/DaySequence/Source/DaySequence/Private/DaySequenceActor.cpp
--- a/Engine/Plugins/Experimental/DaySequence/Source/DaySequence/Private/DaySequenceActor.cpp
+++ b/Engine/Plugins/Experimental/DaySequence/Source/DaySequence/Private/DaySequenceActor.cpp
@@ -1452,7 +1452,7 @@ UDaySequence* ADaySequenceActor::GetRootSequence() const
 	return RootSequence;
 }
 
-void ADaySequenceActor::UpdateRootSequence()
+void ADaySequenceActor::UpdateRootSequence(EUpdateRootSequenceMode Mode)
 {
 #if WITH_EDITOR
 	FCoreUObjectDelegates::OnObjectModified.RemoveAll(this);
@@ -1464,20 +1464,28 @@ void ADaySequenceActor::UpdateRootSequence()
 		{
 			// We update on next tick because calling UpdateScalabilitySequence here
 			// is too early (our collection doesn't have the changes that triggered this invocation).
-			bUpdateRootSequenceOnTick = true;
+			UpdateRootSequenceOnTick(EUpdateRootSequenceMode::Refresh);
 		}
 	});
 #endif
-	
-	// Force rebuild the root sequence.
-	SubSections.Empty();
-	RootSequence = nullptr;
+
+	if (Mode == EUpdateRootSequenceMode::Refresh)
+	{
+		// Force rebuild the root sequence.
+		SubSections.Empty();
+		RootSequence = nullptr;
+	}
 	InitializeRootSequence();
 }
 
 #if WITH_EDITOR
-void ADaySequenceActor::UpdateRootSequenceOnTick()
+void ADaySequenceActor::UpdateRootSequenceOnTick(EUpdateRootSequenceMode Mode)
 {
+	if (Mode == EUpdateRootSequenceMode::Refresh)
+	{
+		SubSections.Empty();
+		RootSequence = nullptr;
+	}
 	bUpdateRootSequenceOnTick = true;
 }
 #endif
diff --git a/Engine/Plugins/Experimental/DaySequence/Source/DaySequence/Public/DaySequenceActor.h b/Engine/Plugins/Experimental/DaySequence/Source/DaySequence/Public/DaySequenceActor.h
--- a/Engine/Plugins/Experimental/DaySequence/Source/DaySequence/Public/DaySequenceActor.h
+++ b/Engine/Plugins/Experimental/DaySequence/Source/DaySequence/Public/DaySequenceActor.h
@@ -314,9 +314,18 @@ public:
 	
 	UDaySequence* GetRootSequence() const;
 
-	virtual void UpdateRootSequence();
+	enum class EUpdateRootSequenceMode
+	{
+		// Compare the current sequence against the new sequence and only add/remove
+		// new/old sequences respectively.
+		Update,
+		// Reset and rebuild the root sequence from scratch.
+		Refresh
+	};
+	
+	virtual void UpdateRootSequence(EUpdateRootSequenceMode Mode = EUpdateRootSequenceMode::Update);
 #if WITH_EDITOR
-	void UpdateRootSequenceOnTick();
+	void UpdateRootSequenceOnTick(EUpdateRootSequenceMode Mode = EUpdateRootSequenceMode::Update);
 #endif
 
 	bool HasValidRootSequence() const;
diff --git a/Engine/Plugins/Experimental/DaySequence/Source/DaySequenceEditor/Private/DaySequenceEditorModule.cpp b/Engine/Plugins/Experimental/DaySequence/Source/DaySequenceEditor/Private/DaySequenceEditorModule.cpp
--- a/Engine/Plugins/Experimental/DaySequence/Source/DaySequenceEditor/Private/DaySequenceEditorModule.cpp
+++ b/Engine/Plugins/Experimental/DaySequence/Source/DaySequenceEditor/Private/DaySequenceEditorModule.cpp
@@ -421,24 +421,13 @@ void FDaySequenceEditorModule::OnOpenRootSequence()
 				GEditor->GetEditorSubsystem<UAssetEditorSubsystem>()->OpenEditorForAsset(LoadedObject);
 
 				// This iteration does work necessary to keep the viewport preview time consistent when opening/closing a sequence editor.
-				FDaySequenceEditorToolkit::IterateOpenToolkits([this, InitialGlobalTime, DayActor](FDaySequenceEditorToolkit& Toolkit)
+				FDaySequenceEditorToolkit::IterateOpenToolkits([InitialGlobalTime](FDaySequenceEditorToolkit& Toolkit)
 				{
 					if (!Toolkit.IsActorPreview())
 					{
-						const TSharedPtr<ISequencer> ToolkitSequencer = Toolkit.GetSequencer();
-						TWeakPtr<ISequencer> WeakToolkitSequencer = ToolkitSequencer;
-						float SequencerTimeConversionFactor = DayActor->GetDayLength() / DayActor->GetTimePerCycle();
-						
 						// One time propagation of preview actor time for initializing the sequence editor
+						const TSharedPtr<ISequencer> ToolkitSequencer = Toolkit.GetSequencer();
 						ToolkitSequencer->SetGlobalTime(InitialGlobalTime);
-
-						// Forward global time updates to the actor preview so the sequence editor time is preserved when closed
-						ToolkitSequencer->OnGlobalTimeChanged().AddWeakLambda(DayActor, [this, WeakToolkitSequencer, SequencerTimeConversionFactor, DayActor]()
-						{
-							const float SequencerTime = WeakToolkitSequencer.Pin()->GetGlobalTime().AsSeconds() / 3600.f;
-							const float PreviewTime = SequencerTime * SequencerTimeConversionFactor;
-							DayActor->ConditionalSetTimeOfDayPreview(PreviewTime);
-						});
 						
 						// Break out of iteration now that we found a sequence editor
 						return false;
@@ -494,7 +483,7 @@ void FDaySequenceEditorModule::OnRefreshDaySequenceActor()
 	const UDaySequenceSubsystem* DaySubsystem = GEditor->GetEditorWorldContext().World()->GetSubsystem<UDaySequenceSubsystem>();
 	if (ADaySequenceActor* DayActor = DaySubsystem->GetDaySequenceActor())
 	{
-		DayActor->UpdateRootSequence();
+		DayActor->UpdateRootSequence(ADaySequenceActor::EUpdateRootSequenceMode::Refresh);
 	}
 }
 
diff --git a/Engine/Plugins/Experimental/DaySequence/Source/DaySequenceEditor/Private/DaySequenceEditorToolkit.cpp b/Engine/Plugins/Experimental/DaySequence/Source/DaySequenceEditor/Private/DaySequenceEditorToolkit.cpp
--- a/Engine/Plugins/Experimental/DaySequence/Source/DaySequenceEditor/Private/DaySequenceEditorToolkit.cpp
+++ b/Engine/Plugins/Experimental/DaySequence/Source/DaySequenceEditor/Private/DaySequenceEditorToolkit.cpp
@@ -167,7 +167,7 @@ FDaySequenceEditorToolkit::~FDaySequenceEditorToolkit()
 		// the preview toolkit. We run the regeneration of the root sequence on tick to avoid this
 		// case.
 		RootActor->bForceDisableDayInterpCurve = false;
-		RootActor->UpdateRootSequenceOnTick();
+		RootActor->UpdateRootSequenceOnTick(ADaySequenceActor::EUpdateRootSequenceMode::Refresh);
 	}
 
 	OnDestroyed().Broadcast(*this);
@@ -786,8 +786,8 @@ void FDaySequenceEditorToolkit::OnGlobalTimeChanged()
 	{
 		// Convert sequencer time to equivalent game time.
 		
-		const FFrameNumber LowerBound = Sequencer->GetPlaybackRange().GetLowerBoundValue();
-		const FFrameNumber UpperBound = Sequencer->GetPlaybackRange().GetUpperBoundValue();
+		const FFrameNumber LowerBound = Sequencer->GetRootMovieSceneSequence()->GetMovieScene()->GetPlaybackRange().GetLowerBoundValue();
+		const FFrameNumber UpperBound = Sequencer->GetRootMovieSceneSequence()->GetMovieScene()->GetPlaybackRange().GetUpperBoundValue();
 		const int32 Range = UpperBound.Value - LowerBound.Value;
 		const int32 CurrentTimeOffset = Sequencer->GetGlobalTime().Time.FrameNumber.Value - LowerBound.Value;
 		const float NormalizedTime = static_cast<float>(CurrentTimeOffset) / Range;
