DaySequence: Fixed flicker on closing the root sequence.

diff --git a/Engine/Plugins/Experimental/DaySequence/Source/DaySequence/Private/DaySequenceActor.cpp b/Engine/Plugins/Experimental/DaySequence/Source/DaySequence/Private/DaySequenceActor.cpp
--- a/Engine/Plugins/Experimental/DaySequence/Source/DaySequence/Private/DaySequenceActor.cpp
+++ b/Engine/Plugins/Experimental/DaySequence/Source/DaySequence/Private/DaySequenceActor.cpp
@@ -172,7 +172,7 @@ void ADaySequenceActor::OnConstruction(const FTransform& Transform)
 {
 	// It is unsafe to update the root sequence (incl. its delegates)
 	// during actor construction. Defer to the next tick.
-	bUpdateRootSequenceOnTick = true;
+	UpdateRootSequenceOnTickState.bUpdate = true;
 }
 #endif
 
@@ -182,10 +182,10 @@ void ADaySequenceActor::Tick(float DeltaTime)
 	Super::Tick(DeltaTime);
 	
 #if WITH_EDITOR
-	if (bUpdateRootSequenceOnTick && GetWorld() != nullptr && GetWorld()->WorldType == EWorldType::Editor)
+	if (UpdateRootSequenceOnTickState.bUpdate && GetWorld() != nullptr && GetWorld()->WorldType == EWorldType::Editor)
 	{
-		UpdateRootSequence();
-		bUpdateRootSequenceOnTick = false;
+		UpdateRootSequence(UpdateRootSequenceOnTickState.Mode);
+		UpdateRootSequenceOnTickState = FUpdateRootSequenceState();
 	}
 #endif
 
@@ -291,7 +291,7 @@ void ADaySequenceActor::PostLoad()
 
 	SubSections.Empty();
 
-	bUpdateRootSequenceOnTick = true;
+	UpdateRootSequenceOnTickState.bUpdate = true;
 #endif
 
 #if WITH_EDITORONLY_DATA
@@ -463,7 +463,7 @@ void ADaySequenceActor::PostEditChangeProperty(FPropertyChangedEvent& PropertyCh
 {
 	auto ReinitializeActor = [this]()
 	{
-		bUpdateRootSequenceOnTick = true;
+		UpdateRootSequenceOnTickState.bUpdate = true;
 		SubSections.Empty();
 	};
 	
@@ -1464,12 +1464,12 @@ void ADaySequenceActor::UpdateRootSequence(EUpdateRootSequenceMode Mode)
 		{
 			// We update on next tick because calling UpdateScalabilitySequence here
 			// is too early (our collection doesn't have the changes that triggered this invocation).
-			UpdateRootSequenceOnTick(EUpdateRootSequenceMode::Refresh);
+			UpdateRootSequenceOnTick(EUpdateRootSequenceMode::Reinitialize);
 		}
 	});
 #endif
 
-	if (Mode == EUpdateRootSequenceMode::Refresh)
+	if (EnumHasAnyFlags(Mode, EUpdateRootSequenceMode::Reinitialize))
 	{
 		// Force rebuild the root sequence.
 		SubSections.Empty();
@@ -1481,12 +1481,8 @@ void ADaySequenceActor::UpdateRootSequence(EUpdateRootSequenceMode Mode)
 #if WITH_EDITOR
 void ADaySequenceActor::UpdateRootSequenceOnTick(EUpdateRootSequenceMode Mode)
 {
-	if (Mode == EUpdateRootSequenceMode::Refresh)
-	{
-		SubSections.Empty();
-		RootSequence = nullptr;
-	}
-	bUpdateRootSequenceOnTick = true;
+	UpdateRootSequenceOnTickState.bUpdate = true;
+	UpdateRootSequenceOnTickState.Mode |= Mode;
 }
 #endif
 
diff --git a/Engine/Plugins/Experimental/DaySequence/Source/DaySequence/Public/DaySequenceActor.h b/Engine/Plugins/Experimental/DaySequence/Source/DaySequence/Public/DaySequenceActor.h
--- a/Engine/Plugins/Experimental/DaySequence/Source/DaySequence/Public/DaySequenceActor.h
+++ b/Engine/Plugins/Experimental/DaySequence/Source/DaySequence/Public/DaySequenceActor.h
@@ -318,9 +318,9 @@ public:
 	{
 		// Compare the current sequence against the new sequence and only add/remove
 		// new/old sequences respectively.
-		Update,
+		Update       = 0,
 		// Reset and rebuild the root sequence from scratch.
-		Refresh
+		Reinitialize = 1 << 0,
 	};
 	
 	virtual void UpdateRootSequence(EUpdateRootSequenceMode Mode = EUpdateRootSequenceMode::Update);
@@ -605,8 +605,13 @@ protected:
 #endif //WITH_EDITORONLY_DATA
 
 #if WITH_EDITOR
+	struct FUpdateRootSequenceState
+	{
+		bool bUpdate = false;
+		EUpdateRootSequenceMode Mode = EUpdateRootSequenceMode::Update;
+	};
 	/** When true, the root sequence will be updated on next Tick. */
-	bool bUpdateRootSequenceOnTick = false;
+	FUpdateRootSequenceState UpdateRootSequenceOnTickState = FUpdateRootSequenceState{false, EUpdateRootSequenceMode::Update};
 #endif
 
 	FOnRootSequenceChanged OnPreRootSequenceChanged;
@@ -649,3 +654,4 @@ protected:
 	TMap<TSubclassOf<UDaySequenceConditionTag>, TObjectPtr<UDaySequenceConditionTag>> TrackConditionMap;
 };
 
+ENUM_CLASS_FLAGS(ADaySequenceActor::EUpdateRootSequenceMode);
diff --git a/Engine/Plugins/Experimental/DaySequence/Source/DaySequenceEditor/Private/DaySequenceEditorModule.cpp b/Engine/Plugins/Experimental/DaySequence/Source/DaySequenceEditor/Private/DaySequenceEditorModule.cpp
--- a/Engine/Plugins/Experimental/DaySequence/Source/DaySequenceEditor/Private/DaySequenceEditorModule.cpp
+++ b/Engine/Plugins/Experimental/DaySequence/Source/DaySequenceEditor/Private/DaySequenceEditorModule.cpp
@@ -483,7 +483,7 @@ void FDaySequenceEditorModule::OnRefreshDaySequenceActor()
 	const UDaySequenceSubsystem* DaySubsystem = GEditor->GetEditorWorldContext().World()->GetSubsystem<UDaySequenceSubsystem>();
 	if (ADaySequenceActor* DayActor = DaySubsystem->GetDaySequenceActor())
 	{
-		DayActor->UpdateRootSequence(ADaySequenceActor::EUpdateRootSequenceMode::Refresh);
+		DayActor->UpdateRootSequence(ADaySequenceActor::EUpdateRootSequenceMode::Reinitialize);
 	}
 }
 
diff --git a/Engine/Plugins/Experimental/DaySequence/Source/DaySequenceEditor/Private/DaySequenceEditorToolkit.cpp b/Engine/Plugins/Experimental/DaySequence/Source/DaySequenceEditor/Private/DaySequenceEditorToolkit.cpp
--- a/Engine/Plugins/Experimental/DaySequence/Source/DaySequenceEditor/Private/DaySequenceEditorToolkit.cpp
+++ b/Engine/Plugins/Experimental/DaySequence/Source/DaySequenceEditor/Private/DaySequenceEditorToolkit.cpp
@@ -167,7 +167,7 @@ FDaySequenceEditorToolkit::~FDaySequenceEditorToolkit()
 		// the preview toolkit. We run the regeneration of the root sequence on tick to avoid this
 		// case.
 		RootActor->bForceDisableDayInterpCurve = false;
-		RootActor->UpdateRootSequenceOnTick(ADaySequenceActor::EUpdateRootSequenceMode::Refresh);
+		RootActor->UpdateRootSequenceOnTick(ADaySequenceActor::EUpdateRootSequenceMode::Reinitialize);
 	}
 
 	OnDestroyed().Broadcast(*this);
