DaySequenceActor: Move playback client assignment to game world only path

diff --git a/Engine/Plugins/Experimental/DaySequence/Source/DaySequence/Private/DaySequenceActor.cpp b/Engine/Plugins/Experimental/DaySequence/Source/DaySequence/Private/DaySequenceActor.cpp
--- a/Engine/Plugins/Experimental/DaySequence/Source/DaySequence/Private/DaySequenceActor.cpp
+++ b/Engine/Plugins/Experimental/DaySequence/Source/DaySequence/Private/DaySequenceActor.cpp
@@ -259,19 +259,6 @@ void ADaySequenceActor::GetLifetimeReplicatedProps(TArray<FLifetimeProperty>& Ou
 	DOREPLIFETIME(ADaySequenceActor, bUseInterpCurve);
 }
 
-void ADaySequenceActor::PostInitProperties()
-{
-	Super::PostInitProperties();
-
-	// This is only checking that the UObject is valid, not that the player has been initialized yet.
-	if (IsValid(SequencePlayer))
-	{
-		// Have to initialize this here as any properties set on default subobjects inside the constructor
-		// get stomped by the CDO's properties when the constructor exits.
-		SequencePlayer->SetPlaybackClient(this);
-	}
-}
-
 void ADaySequenceActor::PostLoad()
 {
 	Super::PostLoad();
@@ -561,6 +548,7 @@ void ADaySequenceActor::InitializePlayer()
 	{
 		if (SequencePlayer)
 		{
+			SequencePlayer->SetPlaybackClient(this);
 			SequencePlayer->Initialize(RootSequence, this, GetPlaybackSettings(RootSequence));
 		}
 
diff --git a/Engine/Plugins/Experimental/DaySequence/Source/DaySequence/Public/DaySequenceActor.h b/Engine/Plugins/Experimental/DaySequence/Source/DaySequence/Public/DaySequenceActor.h
--- a/Engine/Plugins/Experimental/DaySequence/Source/DaySequence/Public/DaySequenceActor.h
+++ b/Engine/Plugins/Experimental/DaySequence/Source/DaySequence/Public/DaySequenceActor.h
@@ -402,7 +402,6 @@ public:
 	
 protected:
 	//~ Begin UObject interface
-	virtual void PostInitProperties() override;
 	virtual void PostLoad() override;
 	virtual void BeginDestroy() override;
 #if WITH_EDITOR
