DaySequence: Fixed issue with DaySequenceCollections upgrade that caused inherited collections to populate into the array.

diff --git a/Engine/Plugins/Experimental/DaySequence/Source/DaySequence/Private/DaySequenceActor.cpp b/Engine/Plugins/Experimental/DaySequence/Source/DaySequence/Private/DaySequenceActor.cpp
--- a/Engine/Plugins/Experimental/DaySequence/Source/DaySequence/Private/DaySequenceActor.cpp
+++ b/Engine/Plugins/Experimental/DaySequence/Source/DaySequence/Private/DaySequenceActor.cpp
@@ -294,6 +294,9 @@ PRAGMA_DISABLE_DEPRECATION_WARNINGS
 	// Transfer DaySequenceCollection (deprecated) to DaySequenceCollections
 	if (DaySequenceCollection)
 	{
+		// When transferring to DaySequenceCollections, clear the array to preserve
+		// child class override behavior of DaySequenceCollection.
+		DaySequenceCollections.Reset();
 		DaySequenceCollections.Add(DaySequenceCollection);
 		DaySequenceCollection = nullptr;
 	}
