Ensure null safety when constructing the legacy Day Sequence toolbar extension.

#rnx

diff --git a/Engine/Plugins/Experimental/DaySequence/Source/DaySequenceEditor/Private/DaySequenceEditorModule.cpp b/Engine/Plugins/Experimental/DaySequence/Source/DaySequenceEditor/Private/DaySequenceEditorModule.cpp
--- a/Engine/Plugins/Experimental/DaySequence/Source/DaySequenceEditor/Private/DaySequenceEditorModule.cpp
+++ b/Engine/Plugins/Experimental/DaySequence/Source/DaySequenceEditor/Private/DaySequenceEditorModule.cpp
@@ -620,18 +620,21 @@ TSharedRef<SWidget> FDaySequenceEditorModule::CreateDaySequenceViewportToolbarEx
 {
 	if (const UViewportToolBarContext* ExtensionContextObject = Cast<UViewportToolBarContext>(ExtensionContext.Get()); ensure(ExtensionContextObject))
 	{
-		return SNew(SEditorViewportToolbarMenu)
-			.ParentToolBar(ExtensionContextObject->ViewportToolBar.Pin())
-			.Label(LOCTEXT("DaySequenceMenuLabel", "Time of Day"))
-			.LabelIcon(FDaySequenceEditorStyle::Get()->GetBrush("DaySequenceEditor.ViewportToolBar"))
-			.OnGetMenuContent_Lambda([this]()
-			{
-				return CreateTimeOfDayWidget();
-			})
-			.Visibility_Lambda([this]()
-			{
-				return IsMenuVisible() && UE::UnrealEd::ShowOldViewportToolbars() ? EVisibility::Visible : EVisibility::Collapsed;
-			});
+		if (TSharedPtr<SViewportToolBar> Toolbar = ExtensionContextObject->ViewportToolBar.Pin())
+		{
+			return SNew(SEditorViewportToolbarMenu)
+				.ParentToolBar(Toolbar)
+				.Label(LOCTEXT("DaySequenceMenuLabel", "Time of Day"))
+				.LabelIcon(FDaySequenceEditorStyle::Get()->GetBrush("DaySequenceEditor.ViewportToolBar"))
+				.OnGetMenuContent_Lambda([this]()
+				{
+					return CreateTimeOfDayWidget();
+				})
+				.Visibility_Lambda([this]()
+				{
+					return IsMenuVisible() && UE::UnrealEd::ShowOldViewportToolbars() ? EVisibility::Visible : EVisibility::Collapsed;
+				});
+		}
 	}
 	return SNullWidget::NullWidget;
 }
