DaySequenceConditionTag: Add initialization phase to condition tag defaults to allow condition tags to defer their own initialization so that game state exists when initializing

diff --git a/Engine/Plugins/Experimental/DaySequence/Source/DaySequence/Private/DaySequenceActor.cpp b/Engine/Plugins/Experimental/DaySequence/Source/DaySequence/Private/DaySequenceActor.cpp
--- a/Engine/Plugins/Experimental/DaySequence/Source/DaySequence/Private/DaySequenceActor.cpp
+++ b/Engine/Plugins/Experimental/DaySequence/Source/DaySequence/Private/DaySequenceActor.cpp
@@ -1760,8 +1760,35 @@ TObjectPtr<UDaySequenceConditionTag> ADaySequenceActor::GetOrInstantiateConditio
 	{
 		TrackConditionMap.Remove(*ConditionClass);
 		TrackConditionMap.Add(*ConditionClass, NewObject<UDaySequenceConditionTag>(this, ConditionClass));
+		UDaySequenceConditionTag* ConditionTag = TrackConditionMap[*ConditionClass];
+		
+		UWorld* World = GetWorld();
+
+		switch (TrackConditionMap[*ConditionClass]->GetInitializationPhase())
+		{
+		case EDaySequenceConditionTagInitializationPhase::Immediate:
+			ConditionTag->Initialize();
+			break;
+
+		case EDaySequenceConditionTagInitializationPhase::GameStateReady:
+			if (World)
+			{
+				World->GameStateSetEvent.AddWeakLambda(ConditionTag, [this, ConditionTag](AGameStateBase*)
+				{
+					ConditionTag->Initialize();
+					InvalidateMuteStates();
+				});
 
-		TrackConditionMap[*ConditionClass]->Initialize();
+				if (World->GetGameState())
+				{
+					ConditionTag->Initialize();
+				}
+			}
+			break;
+			
+		default:
+			break;
+		}
 	}
 
 	return TrackConditionMap[*ConditionClass];
diff --git a/Engine/Plugins/Experimental/DaySequence/Source/DaySequence/Private/DaySequenceConditionTag.cpp b/Engine/Plugins/Experimental/DaySequence/Source/DaySequence/Private/DaySequenceConditionTag.cpp
--- a/Engine/Plugins/Experimental/DaySequence/Source/DaySequence/Private/DaySequenceConditionTag.cpp
+++ b/Engine/Plugins/Experimental/DaySequence/Source/DaySequence/Private/DaySequenceConditionTag.cpp
@@ -17,6 +17,11 @@ FString UDaySequenceConditionTag::GetConditionName() const
 	return ConditionName.IsEmpty() ? GetClass()->GetName() : ConditionName;
 }
 
+EDaySequenceConditionTagInitializationPhase UDaySequenceConditionTag::GetInitializationPhase() const
+{
+	return InitializationPhase;
+}
+
 UWorld* UDaySequenceConditionTag::GetWorld() const
 {
 	if (IsTemplate())
diff --git a/Engine/Plugins/Experimental/DaySequence/Source/DaySequence/Public/DaySequenceConditionTag.h b/Engine/Plugins/Experimental/DaySequence/Source/DaySequence/Public/DaySequenceConditionTag.h
--- a/Engine/Plugins/Experimental/DaySequence/Source/DaySequence/Public/DaySequenceConditionTag.h
+++ b/Engine/Plugins/Experimental/DaySequence/Source/DaySequence/Public/DaySequenceConditionTag.h
@@ -10,6 +10,14 @@
 
 class ADaySequenceActor;
 
+/** Enum specifying when a day sequence condition tag will be initialized. */
+UENUM(BlueprintType)
+enum class EDaySequenceConditionTagInitializationPhase : uint8
+{
+	Immediate,
+	GameStateReady
+};
+
 UCLASS(MinimalAPI, Abstract, Blueprintable, Const)
 class UDaySequenceConditionTag
 	: public UObject
@@ -29,6 +37,7 @@ public:
 	UE_API bool Evaluate() const;
 
 	UE_API FString GetConditionName() const;
+	UE_API EDaySequenceConditionTagInitializationPhase GetInitializationPhase() const;
 	
 	DECLARE_MULTICAST_DELEGATE(FOnConditionValueChanged);
 	FOnConditionValueChanged& GetOnConditionValueChanged() { return OnConditionValueChanged; }
@@ -60,9 +69,15 @@ protected:
 	UPROPERTY(EditDefaultsOnly, Category = "General")
 	FString ConditionName;
 
+	/**
+	 * Determines when SetupOnConditionValueChanged is called.
+	 */
+	UPROPERTY(EditDefaultsOnly, Category = "General")
+	EDaySequenceConditionTagInitializationPhase InitializationPhase = EDaySequenceConditionTagInitializationPhase::Immediate;
+
 private:
 	FOnConditionValueChanged OnConditionValueChanged;
-
+	
 	/* This is an optional because it is unset until the first time BroadcastOnConditionValueChanged is called. */
 	mutable TOptional<bool> bCachedEvalResult;
 };
