DaySequenceConditionTag: Properly detect BP subclasses of native classes which derive from UDaySequenceConditionTag

diff --git a/Engine/Plugins/Experimental/DaySequence/Source/DaySequenceEditor/Private/SDaySequenceConditionSetPicker.cpp b/Engine/Plugins/Experimental/DaySequence/Source/DaySequenceEditor/Private/SDaySequenceConditionSetPicker.cpp
--- a/Engine/Plugins/Experimental/DaySequence/Source/DaySequenceEditor/Private/SDaySequenceConditionSetPicker.cpp
+++ b/Engine/Plugins/Experimental/DaySequence/Source/DaySequenceEditor/Private/SDaySequenceConditionSetPicker.cpp
@@ -183,40 +183,11 @@ void SDaySequenceConditionSetPicker::PopulateVisibleClasses()
 {
 	const IAssetRegistry& AssetRegistry = IAssetRegistry::GetChecked();
 	const FTopLevelAssetPath BaseAssetPath = UDaySequenceConditionTag::StaticClass()->GetClassPathName();
+	TArray<FTopLevelAssetPath> NativeAssetPaths;
+	NativeAssetPaths.AddUnique(BaseAssetPath);
 
 	AllConditionTags.Reset();
 
-	// Scan for BP assets to get generated classes
-	{
-		FARFilter Filter;
-		Filter.ClassPaths.Add(UBlueprint::StaticClass()->GetClassPathName());
-		Filter.ClassPaths.Add(UBlueprintGeneratedClass::StaticClass()->GetClassPathName());
-		Filter.bRecursiveClasses = true;
-		Filter.TagsAndValues.Add(FBlueprintTags::NativeParentClassPath);
-	
-		AssetRegistry.EnumerateAssets(Filter, [this, BaseAssetPath](const FAssetData& Data)
-		{
-			if (Data.TagsAndValues.FindTag(FBlueprintTags::NativeParentClassPath).AsExportPath().ToTopLevelAssetPath() == BaseAssetPath)
-			{
-				if (const UBlueprint* AssetAsBlueprint = Cast<UBlueprint>(Data.GetAsset()))
-				{
-					if (const TSubclassOf<UDaySequenceConditionTag> GeneratedClassAsConditionClass = *AssetAsBlueprint->GeneratedClass)
-					{
-						AllConditionTags.AddUnique(GeneratedClassAsConditionClass);
-					}
-				}
-				else if (UBlueprintGeneratedClass* AssetAsBlueprintGeneratedClass = Cast<UBlueprintGeneratedClass>(Data.GetAsset()))
-				{
-					if (AssetAsBlueprintGeneratedClass->IsChildOf(UDaySequenceConditionTag::StaticClass()))
-					{
-						AllConditionTags.Add(AssetAsBlueprintGeneratedClass);
-					}
-				}
-			}
-			return true;	// Returning false will halt the enumeration
-		});
-	}
-
 	// Scan for native classes
 	{
 		const UClass* ParentClass = UDaySequenceConditionTag::StaticClass();
@@ -233,9 +204,45 @@ void SDaySequenceConditionSetPicker::PopulateVisibleClasses()
 			if (UClass* CurrentClass = *It; CurrentClass && IsNativeSubclass(CurrentClass))
 			{
 				AllConditionTags.Add(CurrentClass);
+				const FTopLevelAssetPath CurrentClassPath(CurrentClass->GetPathName());
+				NativeAssetPaths.Add(CurrentClassPath);
 			}
 		}
 	}
+	
+	// Scan for BP assets to get generated classes
+    {
+    	FARFilter Filter;
+    	Filter.ClassPaths.Add(UBlueprint::StaticClass()->GetClassPathName());
+    	Filter.ClassPaths.Add(UBlueprintGeneratedClass::StaticClass()->GetClassPathName());
+    	Filter.bRecursiveClasses = true;
+    	Filter.TagsAndValues.Add(FBlueprintTags::NativeParentClassPath);
+    
+    	AssetRegistry.EnumerateAssets(Filter, [this, NativeAssetPaths](const FAssetData& Data)
+    	{
+    		for (auto CurrentAssetPath : NativeAssetPaths)
+    		{
+    			if (Data.TagsAndValues.FindTag(FBlueprintTags::NativeParentClassPath).AsExportPath().ToTopLevelAssetPath() == CurrentAssetPath)
+    			{
+					if (const UBlueprint* AssetAsBlueprint = Cast<UBlueprint>(Data.GetAsset()))
+					{
+						if (const TSubclassOf<UDaySequenceConditionTag> GeneratedClassAsConditionClass = *AssetAsBlueprint->GeneratedClass)
+						{
+							AllConditionTags.AddUnique(GeneratedClassAsConditionClass);
+						}
+					}
+					else if (UBlueprintGeneratedClass* AssetAsBlueprintGeneratedClass = Cast<UBlueprintGeneratedClass>(Data.GetAsset()))
+					{
+						if (AssetAsBlueprintGeneratedClass->IsChildOf(UDaySequenceConditionTag::StaticClass()))
+						{
+							AllConditionTags.Add(AssetAsBlueprintGeneratedClass);
+						}
+					}
+				}
+    		}
+    		return true;	// Returning false will halt the enumeration
+    	});
+    }
 }
 
 void SDaySequenceConditionSetPicker::PopulateCheckedTags()
