DaySequence: Allow update interval override to force synchronous evaluations.

1) Interpret update interval of 0 as synchronous updates (budget of 0)
2) Make sure the update interval handle is created (this was a bug)

diff --git a/Engine/Plugins/Experimental/DaySequence/Source/DaySequence/Private/DaySequenceModifierComponent.cpp b/Engine/Plugins/Experimental/DaySequence/Source/DaySequence/Private/DaySequenceModifierComponent.cpp
--- a/Engine/Plugins/Experimental/DaySequence/Source/DaySequence/Private/DaySequenceModifierComponent.cpp
+++ b/Engine/Plugins/Experimental/DaySequence/Source/DaySequence/Private/DaySequenceModifierComponent.cpp
@@ -275,6 +275,14 @@ void UDaySequenceModifierComponent::DaySequenceUpdate()
 	CSV_SCOPED_TIMING_STAT(DaySequence, SequencePlayerUpdated);
 
 	bool bWantsOverride = false;
+
+	if (!OverrideUpdateIntervalHandle && TargetActor)
+	{
+		if (IDaySequencePlayer* Player = TargetActor->GetSequencePlayer())
+		{
+			OverrideUpdateIntervalHandle = Player->GetOverrideUpdateIntervalHandle();
+		}
+	}
 	
 	if (bIsComponentEnabled)
 	{
diff --git a/Engine/Plugins/Experimental/DaySequence/Source/DaySequence/Private/DaySequencePlayer.cpp b/Engine/Plugins/Experimental/DaySequence/Source/DaySequence/Private/DaySequencePlayer.cpp
--- a/Engine/Plugins/Experimental/DaySequence/Source/DaySequence/Private/DaySequencePlayer.cpp
+++ b/Engine/Plugins/Experimental/DaySequence/Source/DaySequence/Private/DaySequencePlayer.cpp
@@ -2107,14 +2107,14 @@ void UDaySequencePlayer::Tick(float DeltaSeconds)
 		{
 			// Begin a queued evaluation
 			SCOPE_CYCLE_COUNTER(DaySequenceEval_BeginEvaluation);
-			Runner->Flush(DesiredBudgetMs);
+			Runner->Flush(FMath::IsNearlyZero(EffectiveUpdateInterval) ? 0.f : DesiredBudgetMs);
 		}
 	}
 	else if (bLinkerIsOutstanding)
 	{
 		// Flush if we are in the middle of a budgeted evaluation
 		SCOPE_CYCLE_COUNTER(DaySequenceEval_ContinueEvaluation);
-		Runner->Flush(DesiredBudgetMs);
+		Runner->Flush(FMath::IsNearlyZero(EffectiveUpdateInterval) ? 0.f : DesiredBudgetMs);
 	}
 }
 
