DaySequenceModifierVolume: Disable collision on default box.

diff --git a/Engine/Plugins/Experimental/DaySequence/Source/DaySequence/Private/Actors/DaySequenceModifierVolume.cpp b/Engine/Plugins/Experimental/DaySequence/Source/DaySequence/Private/Actors/DaySequenceModifierVolume.cpp
--- a/Engine/Plugins/Experimental/DaySequence/Source/DaySequence/Private/Actors/DaySequenceModifierVolume.cpp
+++ b/Engine/Plugins/Experimental/DaySequence/Source/DaySequence/Private/Actors/DaySequenceModifierVolume.cpp
@@ -34,6 +34,9 @@ ADaySequenceModifierVolume::ADaySequenceModifierVolume(const FObjectInitializer&
 	DefaultBox->SetupAttachment(DaySequenceModifier);
 	DefaultBox->SetLineThickness(10.f);
 	DefaultBox->SetBoxExtent(FVector(500.f));
+	DefaultBox->SetGenerateOverlapEvents(false);
+	DefaultBox->SetCollisionProfileName(TEXT("NoCollision"));
+	DefaultBox->CanCharacterStepUpOn = ECB_No;
 
 	FComponentReference DefaultBoxReference;
 	DefaultBoxReference.ComponentProperty = TEXT("DefaultBox");
