DaySequence: Moved the Actor level StaticTimeMode/LocalTimeWarp to the ModifierComponent as LocalFixedTime.

This allows per modifier component control over local time warp evaluation rather than requiring all modifier components to be controlled in this manner.

diff --git a/Engine/Plugins/Experimental/DaySequence/Source/DaySequence/Private/DaySequenceActor.cpp b/Engine/Plugins/Experimental/DaySequence/Source/DaySequence/Private/DaySequenceActor.cpp
--- a/Engine/Plugins/Experimental/DaySequence/Source/DaySequence/Private/DaySequenceActor.cpp
+++ b/Engine/Plugins/Experimental/DaySequence/Source/DaySequence/Private/DaySequenceActor.cpp
@@ -487,10 +487,6 @@ void ADaySequenceActor::PostEditChangeProperty(FPropertyChangedEvent& PropertyCh
 	{
 		ReinitializeActor();
 	}
-	else if (PropertyName == GET_MEMBER_NAME_CHECKED(ADaySequenceActor, StaticTimeMode))
-	{
-		ReinitializeActor();
-	}
 	else if (PropertyChangedEvent.ChangeType == EPropertyChangeType::Unspecified)
 	{
 		// This handles undo/redo transactions.
@@ -1146,7 +1142,7 @@ void ADaySequenceActor::WarpEvaluationRange(FMovieSceneEvaluationRange& InOutRan
 	};
 
 	// Warp with static time if necessary
-	if (HasStaticTimeOfDay() && StaticTimeMode == EDaySequenceStaticTimeMode::RootTimeWarp)
+	if (HasStaticTimeOfDay())
 	{
 		const float StaticTimeInGameHours = GetStaticTimeOfDay();
 
@@ -1336,16 +1332,6 @@ float ADaySequenceActor::GetStaticTimeOfDay() const
 	return std::numeric_limits<float>::lowest();
 }
 
-void ADaySequenceActor::SetStaticTimeMode(EDaySequenceStaticTimeMode Mode)
-{
-	StaticTimeMode = Mode;
-}
-
-EDaySequenceStaticTimeMode ADaySequenceActor::GetStaticTimeMode() const
-{
-	return StaticTimeMode;
-}
-
 void ADaySequenceActor::RegisterStaticTimeContributor(const UE::DaySequence::FStaticTimeContributor& NewContributor) const
 {
 	if (!StaticTimeManager)
diff --git a/Engine/Plugins/Experimental/DaySequence/Source/DaySequence/Private/DaySequenceModifierComponent.cpp b/Engine/Plugins/Experimental/DaySequence/Source/DaySequence/Private/DaySequenceModifierComponent.cpp
--- a/Engine/Plugins/Experimental/DaySequence/Source/DaySequence/Private/DaySequenceModifierComponent.cpp
+++ b/Engine/Plugins/Experimental/DaySequence/Source/DaySequence/Private/DaySequenceModifierComponent.cpp
@@ -616,15 +616,17 @@ void UDaySequenceModifierComponent::SetInitialTimeOfDay()
 			}
 			break;
 
+		case EDayNightCycleMode::LocalFixedTime:
+			break;
+
 		case EDayNightCycleMode::StartAtSpecifiedTime:
 		case EDayNightCycleMode::RandomStartTime:
-
-		if (!bHasAuthority && Mode != EDaySequenceModifierMode::Volume)
-		{
-			// Never set initial time of day from non-volume based modifiers if they don't have authority and aren't setting static time.
-			// We'll just get the initial time of day from the server replication.
-			return;
-		}
+			if (!bHasAuthority && Mode != EDaySequenceModifierMode::Volume)
+			{
+				// Never set initial time of day from non-volume based modifiers if they don't have authority and aren't setting static time.
+				// We'll just get the initial time of day from the server replication.
+				return;
+			}
 			
 			TargetActor->SetTimeOfDay(Time);
 #if WITH_EDITOR
@@ -635,6 +637,7 @@ void UDaySequenceModifierComponent::SetInitialTimeOfDay()
 #if WITH_EDITOR
 		default:
 			TargetActor->SetTimeOfDayPreview(TargetActor->GetTimeOfDayPreview());
+			break;
 #endif
 		}
 	}
@@ -788,7 +791,7 @@ UMovieSceneSubSection* UDaySequenceModifierComponent::InitializeDaySequence(cons
 		SubSection->SetRange(PlaybackRange);
 		SubSection->SetIsActive(bActivate);
 
-		if (TargetActor->GetStaticTimeMode() == EDaySequenceStaticTimeMode::LocalTimeWarp)
+		if (DayNightCycle == EDayNightCycleMode::LocalFixedTime)
 		{
 			if (FMovieSceneTimeWarpVariant* Variant = SubSection->GetTimeWarp())
 			{
diff --git a/Engine/Plugins/Experimental/DaySequence/Source/DaySequence/Public/DaySequenceActor.h b/Engine/Plugins/Experimental/DaySequence/Source/DaySequence/Public/DaySequenceActor.h
--- a/Engine/Plugins/Experimental/DaySequence/Source/DaySequence/Public/DaySequenceActor.h
+++ b/Engine/Plugins/Experimental/DaySequence/Source/DaySequence/Public/DaySequenceActor.h
@@ -78,16 +78,6 @@ namespace UE::DaySequence
 
 DECLARE_DYNAMIC_MULTICAST_DELEGATE_OneParam(FOnTimeOfDayPreviewChanged, float, Time);
 
-UENUM()
-enum class EDaySequenceStaticTimeMode
-{
-	// Static time will apply a time warp on the RootSequence/global evaluation time.
-	RootTimeWarp,
-	// Static time will apply a time warp local to the modifier's sequence evaluation time.
-	// The local time warp will still influence the apparent time, but not the global time.
-	LocalTimeWarp
-};
-
 UCLASS(MinimalAPI, Blueprintable)
 class ADaySequenceActor
 	: public AInfo
@@ -303,9 +293,6 @@ public:
 	UFUNCTION(BlueprintCallable, Category="TimeOfDay")
 	UE_API float GetStaticTimeOfDay() const;
 
-	UE_API void SetStaticTimeMode(EDaySequenceStaticTimeMode Mode);
-	UE_API EDaySequenceStaticTimeMode GetStaticTimeMode() const;
-
 	UE_API void RegisterStaticTimeContributor(const UE::DaySequence::FStaticTimeContributor& NewContributor) const;
 	UE_API void UnregisterStaticTimeContributor(const UObject* InUserObject) const;
 	
@@ -589,9 +576,6 @@ protected:
 	
 	TSharedPtr<UE::DaySequence::FStaticTimeManager> StaticTimeManager;
 
-	UPROPERTY(EditAnywhere, Category = RuntimeDayCycle)
-	EDaySequenceStaticTimeMode StaticTimeMode = EDaySequenceStaticTimeMode::RootTimeWarp;
-
 	UPROPERTY(Transient)
 	TObjectPtr<UDaySequenceCameraModifierManager> CameraModifierManager;
 	
diff --git a/Engine/Plugins/Experimental/DaySequence/Source/DaySequence/Public/DaySequenceModifierComponent.h b/Engine/Plugins/Experimental/DaySequence/Source/DaySequence/Public/DaySequenceModifierComponent.h
--- a/Engine/Plugins/Experimental/DaySequence/Source/DaySequence/Public/DaySequenceModifierComponent.h
+++ b/Engine/Plugins/Experimental/DaySequence/Source/DaySequence/Public/DaySequenceModifierComponent.h
@@ -45,6 +45,11 @@ enum class EDayNightCycleMode : uint8
 	RandomFixedTime,
 	/** Start the day/night cycle at a random time, and allow it to continue from there */
 	RandomStartTime,
+	/**
+	 * Apply a time warp local to the modifier's sequence evaluation time.
+	 * Local time warp will not affect global time.
+	 */ 
+	LocalFixedTime
 };
 
 /** Enum that defines modifier behavior for auto enabling and computing the internal blend weight. */
