DaySequence: Fix non-CDO sync load in Base and SunMoon DSA constructors

diff --git a/Engine/Plugins/Experimental/DaySequence/Source/DaySequence/Private/Actors/BaseDaySequenceActor.cpp b/Engine/Plugins/Experimental/DaySequence/Source/DaySequence/Private/Actors/BaseDaySequenceActor.cpp
--- a/Engine/Plugins/Experimental/DaySequence/Source/DaySequence/Private/Actors/BaseDaySequenceActor.cpp
+++ b/Engine/Plugins/Experimental/DaySequence/Source/DaySequence/Private/Actors/BaseDaySequenceActor.cpp
@@ -40,13 +40,6 @@ ABaseDaySequenceActor::ABaseDaySequenceActor(const FObjectInitializer& Init)
 	VolumetricCloudComponent = CreateOptionalDefaultSubobject<UVolumetricCloudComponent>(TEXT("VolumetricCloud"));
     VolumetricCloudComponent->SetupAttachment(RootComponent);
 
-	// We don't want to load this material for the CDO as it will hold on to it forever and it is quite a large asset.
-	if (!IsTemplate())
-	{
-		static ConstructorHelpers::FObjectFinder<UMaterialInterface> VolumetricCloudDefaultMaterialRef(TEXT("/Engine/EngineSky/VolumetricClouds/m_SimpleVolumetricCloud_Inst.m_SimpleVolumetricCloud_Inst"));
-		VolumetricCloudComponent->SetMaterial(VolumetricCloudDefaultMaterialRef.Object);
-	}
-
 	static ConstructorHelpers::FObjectFinder<UStaticMesh> SkySphereDefaultMesh(TEXT("/Engine/EngineSky/SM_SkySphere.SM_SkySphere"));
 	static ConstructorHelpers::FObjectFinder<UMaterialInterface> SkySphereDefaultMaterial(TEXT("/Engine/EngineSky/M_SimpleSkyDome.M_SimpleSkyDome"));
 	SkySphereComponent = CreateOptionalDefaultSubobject<UStaticMeshComponent>(TEXT("SkySphere"));
diff --git a/Engine/Plugins/Experimental/DaySequence/Source/DaySequence/Private/Actors/SunMoonDaySequenceActor.cpp b/Engine/Plugins/Experimental/DaySequence/Source/DaySequence/Private/Actors/SunMoonDaySequenceActor.cpp
--- a/Engine/Plugins/Experimental/DaySequence/Source/DaySequence/Private/Actors/SunMoonDaySequenceActor.cpp
+++ b/Engine/Plugins/Experimental/DaySequence/Source/DaySequence/Private/Actors/SunMoonDaySequenceActor.cpp
@@ -30,10 +30,7 @@ ASunMoonDaySequenceActor::ASunMoonDaySequenceActor(const FObjectInitializer& Ini
 	static ConstructorHelpers::FObjectFinder<UMaterialInterface> SkySphereDefaultMaterial(TEXT("/DaySequence/MI_24hrSky.MI_24hrSky"));
 	SkySphereComponent->SetMaterial(0, SkySphereDefaultMaterial.Object.Get());
 
-	if (!IsTemplate())
-	{
-		// Override the default collection (which animates the moon and sky material)
-        static ConstructorHelpers::FObjectFinder<UDaySequenceCollectionAsset> DefaultCollection(TEXT("/DaySequence/DSCA_24hr.DSCA_24hr"));
-        DaySequenceCollections.Add(DefaultCollection.Object.Get());
-	}
+	// Override the default collection (which animates the moon and sky material)
+	static ConstructorHelpers::FObjectFinder<UDaySequenceCollectionAsset> DefaultCollection(TEXT("/DaySequence/DSCA_24hr.DSCA_24hr"));
+	DaySequenceCollections.Add(DefaultCollection.Object.Get());
 }
\ No newline at end of file
