DaySequence: Fixed issue where upgrading DaySequenceCollection --> DaySequenceCollections was preventing BP derived class overrides on editor load.

diff --git a/Engine/Plugins/Experimental/DaySequence/Source/DaySequence/Private/DaySequenceActor.cpp b/Engine/Plugins/Experimental/DaySequence/Source/DaySequence/Private/DaySequenceActor.cpp
--- a/Engine/Plugins/Experimental/DaySequence/Source/DaySequence/Private/DaySequenceActor.cpp
+++ b/Engine/Plugins/Experimental/DaySequence/Source/DaySequence/Private/DaySequenceActor.cpp
@@ -289,18 +289,6 @@ void ADaySequenceActor::PostLoad()
 		LocalSpriteComponent->SetupAttachment(RootComponent);
 	}
 #endif
-	
-PRAGMA_DISABLE_DEPRECATION_WARNINGS
-	// Transfer DaySequenceCollection (deprecated) to DaySequenceCollections
-	if (DaySequenceCollection)
-	{
-		// When transferring to DaySequenceCollections, clear the array to preserve
-		// child class override behavior of DaySequenceCollection.
-		DaySequenceCollections.Reset();
-		DaySequenceCollections.Add(DaySequenceCollection);
-		DaySequenceCollection = nullptr;
-	}
-PRAGMA_ENABLE_DEPRECATION_WARNINGS
 }
 
 void ADaySequenceActor::BeginDestroy()
@@ -388,6 +376,27 @@ void ADaySequenceActor::Destroyed()
 	Super::Destroyed();
 }
 
+void ADaySequenceActor::Serialize(FArchive& Ar)
+{
+	Super::Serialize(Ar);
+	
+	if (Ar.IsLoading())
+	{
+PRAGMA_DISABLE_DEPRECATION_WARNINGS
+		// Transfer DaySequenceCollection (deprecated) to DaySequenceCollections
+		if (DaySequenceCollection)
+		{
+			// When transferring to DaySequenceCollections, clear the array to preserve
+			// child class override behavior of DaySequenceCollection.
+			DaySequenceCollections.Reset();
+			DaySequenceCollections.Add(DaySequenceCollection);
+			DaySequenceCollection = nullptr;
+		}
+PRAGMA_ENABLE_DEPRECATION_WARNINGS
+	}
+}
+
+
 bool ADaySequenceActor::RetrieveBindingOverrides(const FGuid& InBindingId, FMovieSceneSequenceID InSequenceID, TArray<UObject*, TInlineAllocator<1>>& OutObjects) const
 {
 	if (const TMap<FGuid, TFunction<bool(TArray<UObject*, TInlineAllocator<1>>&)>>* Entry = BindingResolveFunctions.Find(InSequenceID))
diff --git a/Engine/Plugins/Experimental/DaySequence/Source/DaySequence/Public/DaySequenceActor.h b/Engine/Plugins/Experimental/DaySequence/Source/DaySequence/Public/DaySequenceActor.h
--- a/Engine/Plugins/Experimental/DaySequence/Source/DaySequence/Public/DaySequenceActor.h
+++ b/Engine/Plugins/Experimental/DaySequence/Source/DaySequence/Public/DaySequenceActor.h
@@ -404,6 +404,7 @@ protected:
 	UE_API virtual void EndPlay(const EEndPlayReason::Type EndPlayReason) override;
 	UE_API virtual void RewindForReplay() override;
 	UE_API virtual void Destroyed() override;
+	UE_API virtual void Serialize(FArchive& Ar) override;
 
 #if WITH_EDITOR
 	virtual bool CanChangeIsSpatiallyLoadedFlag() const override { return false; }
