DaySequence: Fixed DayNightCycleTime not showing in the details view when using LocalFixedTime.

diff --git a/Engine/Plugins/Experimental/DaySequence/Source/DaySequence/Public/DaySequenceModifierComponent.h b/Engine/Plugins/Experimental/DaySequence/Source/DaySequence/Public/DaySequenceModifierComponent.h
--- a/Engine/Plugins/Experimental/DaySequence/Source/DaySequence/Public/DaySequenceModifierComponent.h
+++ b/Engine/Plugins/Experimental/DaySequence/Source/DaySequence/Public/DaySequenceModifierComponent.h
@@ -294,7 +294,7 @@ protected:
 	int32 Bias;
 
 	/** The time to use for the day/night cycle. */
-	UPROPERTY(EditAnywhere, BlueprintReadWrite, Setter, Getter, DisplayName="Time", Category="Time", meta=(DisplayAfter="DayNightCycle", EditCondition="DayNightCycle==EDayNightCycleMode::FixedTime || DayNightCycle==EDayNightCycleMode::StartAtSpecifiedTime", EditConditionHides))
+	UPROPERTY(EditAnywhere, BlueprintReadWrite, Setter, Getter, DisplayName="Time", Category="Time", meta=(DisplayAfter="DayNightCycle", EditCondition="DayNightCycle==EDayNightCycleMode::FixedTime || DayNightCycle==EDayNightCycleMode::StartAtSpecifiedTime || DayNightCycle==EDayNightCycleMode::LocalFixedTime", EditConditionHides))
 	float DayNightCycleTime;
 
 	/** Defines the region in which the effective blend weight is in the range (0.0, 1.0) (not inclusive) when Mode == EDaySequenceModifierMode::Volume. */
