Fix Sequence Modifier Component to only tick when enabled, which fixes a race condition during creation in async loading

diff --git a/Engine/Plugins/Experimental/DaySequence/Source/DaySequence/Private/DaySequenceModifierComponent.cpp b/Engine/Plugins/Experimental/DaySequence/Source/DaySequence/Private/DaySequenceModifierComponent.cpp
--- a/Engine/Plugins/Experimental/DaySequence/Source/DaySequence/Private/DaySequenceModifierComponent.cpp
+++ b/Engine/Plugins/Experimental/DaySequence/Source/DaySequence/Private/DaySequenceModifierComponent.cpp
@@ -260,6 +260,13 @@ void UDaySequenceModifierComponent::OnRegister()
 	Super::OnRegister();
 
 	bCachedExternalShapesInvalid = true;
+
+#if WITH_EDITOR
+	if (GIsEditor)
+	{
+		SetTickableTickType(GetTickableTickType());
+	}
+#endif
 }
 
 void UDaySequenceModifierComponent::OnUnregister()
@@ -270,6 +277,13 @@ void UDaySequenceModifierComponent::OnUnregister()
 	
 	DisableModifier();
 	RemoveSubSequenceTrack();
+
+#if WITH_EDITOR
+	if (GIsEditor)
+	{
+		SetTickableTickType(ETickableTickType::Never);
+	}
+#endif
 }
 
 void UDaySequenceModifierComponent::DaySequenceUpdate()
diff --git a/Engine/Plugins/Experimental/DaySequence/Source/DaySequence/Public/DaySequenceModifierComponent.h b/Engine/Plugins/Experimental/DaySequence/Source/DaySequence/Public/DaySequenceModifierComponent.h
--- a/Engine/Plugins/Experimental/DaySequence/Source/DaySequence/Public/DaySequenceModifierComponent.h
+++ b/Engine/Plugins/Experimental/DaySequence/Source/DaySequence/Public/DaySequenceModifierComponent.h
@@ -88,8 +88,12 @@ enum class EDaySequenceModifierUserBlendPolicy : uint8
 	 * Editor-only tickable class that allows us to enable trigger volume previews based on
 	 * persepective camera position in the level viewport.
 	 */
-	struct FDaySequenceModifierComponentTickableBase : FTickableGameObject
+	struct FDaySequenceModifierComponentTickableBase : public FTickableGameObject
 	{
+		FDaySequenceModifierComponentTickableBase()
+			: FTickableGameObject(ETickableTickType::Never)
+		{}
+
 		virtual void UpdateEditorPreview(float DeltaTime)
 		{}
 
