DaySequence: Fixed crash in DaySequenceModifierComponent when setting DayNightCycle to LocalFixedTime with a null sequence/collection.

diff --git a/Engine/Plugins/Experimental/DaySequence/Source/DaySequence/Private/DaySequenceModifierComponent.cpp b/Engine/Plugins/Experimental/DaySequence/Source/DaySequence/Private/DaySequenceModifierComponent.cpp
--- a/Engine/Plugins/Experimental/DaySequence/Source/DaySequence/Private/DaySequenceModifierComponent.cpp
+++ b/Engine/Plugins/Experimental/DaySequence/Source/DaySequence/Private/DaySequenceModifierComponent.cpp
@@ -810,16 +810,18 @@ UMovieSceneSubSection* UDaySequenceModifierComponent::InitializeDaySequence(cons
 		{
 			if (FMovieSceneTimeWarpVariant* Variant = SubSection->GetTimeWarp())
 			{
-				UMovieScene* SubMovieScene = Sequence->GetMovieScene();
-				const FFrameRate TickResolution = SubMovieScene->GetTickResolution();
-				const FQualifiedFrameTime SubDuration = FQualifiedFrameTime(
-					UE::MovieScene::DiscreteSize(SubMovieScene->GetPlaybackRange()),
-					TickResolution);
-
-				const float TargetFixedRatio = FMath::Fmod(DayNightCycleTime / TargetActor->GetDayLength(), 1.0f);
-				const FTimecode TargetFixedTimecode = FTimecode(TargetFixedRatio * SubDuration.AsSeconds(), TickResolution, false, false);
-				const FFrameNumber TargetFrame = TargetFixedTimecode.ToFrameNumber(TickResolution);
-				Variant->Set(FMovieSceneTimeWarpFixedFrame{TargetFrame});
+				if (UMovieScene* SubMovieScene = Sequence ? Sequence->GetMovieScene() : nullptr)
+				{
+					const FFrameRate TickResolution = SubMovieScene->GetTickResolution();
+					const FQualifiedFrameTime SubDuration = FQualifiedFrameTime(
+						UE::MovieScene::DiscreteSize(SubMovieScene->GetPlaybackRange()),
+						TickResolution);
+
+					const float TargetFixedRatio = FMath::Fmod(DayNightCycleTime / TargetActor->GetDayLength(), 1.0f);
+					const FTimecode TargetFixedTimecode = FTimecode(TargetFixedRatio * SubDuration.AsSeconds(), TickResolution, false, false);
+					const FFrameNumber TargetFrame = TargetFixedTimecode.ToFrameNumber(TickResolution);
+					Variant->Set(FMovieSceneTimeWarpFixedFrame{TargetFrame});
+				}
 			}
 		}
 		else
