DaySequence: Fixed crash when editing DayLength / TimePerCycle during PIE.

#Synced-CL 48909380

diff --git a/Engine/Plugins/Experimental/DaySequence/Source/DaySequence/Private/DaySequenceActor.cpp b/Engine/Plugins/Experimental/DaySequence/Source/DaySequence/Private/DaySequenceActor.cpp
--- a/Engine/Plugins/Experimental/DaySequence/Source/DaySequence/Private/DaySequenceActor.cpp
+++ b/Engine/Plugins/Experimental/DaySequence/Source/DaySequence/Private/DaySequenceActor.cpp
@@ -460,8 +460,17 @@ bool ADaySequenceActor::GetReferencedContentObjects(TArray<UObject*>& Objects) c
 
 void ADaySequenceActor::PostEditChangeProperty(FPropertyChangedEvent& PropertyChangedEvent)
 {
-	auto ReinitializeActor = [this]()
+	auto ReinitializeActor = [this](bool bReinitializeRootSequence = false)
 	{
+		// Do not reinitialize actor for PIE sessions.
+		if (UWorld* World = GetWorld(); World && World->IsGameWorld())
+		{
+			return;
+		}
+		if (bReinitializeRootSequence)
+		{
+			RootSequence = nullptr;
+		}
 		UpdateRootSequenceOnTickState.bUpdate = true;
 		SubSections.Empty();
 	};
@@ -476,17 +485,15 @@ void ADaySequenceActor::PostEditChangeProperty(FPropertyChangedEvent& PropertyCh
 	{
 		SetDayLength(GetDayLength());
 
-		// RootSequence playback range is computed using DayLength, so lets reconstruct RootSequence to reflect this change.
-		RootSequence = nullptr;
-		ReinitializeActor();
+		// RootSequence playback range is computed using DayLength, so reconstruct RootSequence to reflect this change.
+		ReinitializeActor(true);
 	}
 	else if (PropertyName == GET_MEMBER_NAME_CHECKED(ADaySequenceActor, TimePerCycle))
 	{
 		SetTimePerCycle(GetTimePerCycle());
 
-		// need to null this out to guarantee total reconstruction.
-		RootSequence = nullptr;
-		ReinitializeActor();
+		// RootSequence playback range is computed using TimePerCycle, so reconstruct RootSequence to reflect this change.
+		ReinitializeActor(true);
 	}
 	else if (PropertyName == GET_MEMBER_NAME_CHECKED(ADaySequenceActor, InitialTimeOfDay))
 	{
