DaySequenceModifierVolume: Disable managed modifier components when >1 local players are detected and the volume does not have splitscreen support enabled

diff --git a/Engine/Plugins/Experimental/DaySequence/Source/DaySequence/Private/Actors/DaySequenceModifierVolume.cpp b/Engine/Plugins/Experimental/DaySequence/Source/DaySequence/Private/Actors/DaySequenceModifierVolume.cpp
--- a/Engine/Plugins/Experimental/DaySequence/Source/DaySequence/Private/Actors/DaySequenceModifierVolume.cpp
+++ b/Engine/Plugins/Experimental/DaySequence/Source/DaySequence/Private/Actors/DaySequenceModifierVolume.cpp
@@ -215,7 +215,7 @@ void ADaySequenceModifierVolume::CreatePlayer(APlayerController* InPC)
 		return;
 	}
 
-	UDaySequenceModifierComponent* PlayerModifier;
+	UDaySequenceModifierComponent* PlayerModifier = nullptr;
 	
 	if (CachedPlayerController == nullptr)
 	{
@@ -229,11 +229,11 @@ void ADaySequenceModifierVolume::CreatePlayer(APlayerController* InPC)
 	}
 	else
 	{
-		// If CachedPlayerController is not null (points to the first local player) and we don't support splitscreen, just bail out.
-		return;
+		// If CachedPlayerController is not null (points to the first local player) and we don't support splitscreen, just track the player.
+		AdditionalPlayers.FindOrAdd(InPC) = nullptr;
 	}
 
-	if (ensureMsgf(PlayerModifier, TEXT("PlayerModifier is nullptr!")))
+	if (PlayerModifier)
 	{
 		if (!PlayerModifier->IsRegistered())
 		{
@@ -246,6 +246,24 @@ void ADaySequenceModifierVolume::CreatePlayer(APlayerController* InPC)
 		PlayerModifier->SetBlendTarget(InPC);
 		PlayerModifier->BindToDaySequenceActor(DaySequenceActor);
 	}
+	else if (!bEnableSplitscreenSupport)
+	{
+		// If we get here, we have multiple local players and do not have splitscreen support enabled.
+		// In this case, we will just disable all modifiers.
+		
+		for (TTuple AdditionalPlayer : AdditionalPlayers)
+		{
+			if (AdditionalPlayer.Value)
+			{
+				AdditionalPlayer.Value->DisableComponent();
+			}
+		}
+		
+		if (DaySequenceModifier)
+		{
+			DaySequenceModifier->DisableComponent();
+		}
+	}
 }
 
 bool ADaySequenceModifierVolume::IsSplitscreenSupported() const
