DaySequence: Use GetBlendWeight() to determine the enable state of DaySequenceModifierComponent.

diff --git a/Engine/Plugins/Experimental/DaySequence/Source/DaySequence/Private/DaySequenceModifierComponent.cpp b/Engine/Plugins/Experimental/DaySequence/Source/DaySequence/Private/DaySequenceModifierComponent.cpp
--- a/Engine/Plugins/Experimental/DaySequence/Source/DaySequence/Private/DaySequenceModifierComponent.cpp
+++ b/Engine/Plugins/Experimental/DaySequence/Source/DaySequence/Private/DaySequenceModifierComponent.cpp
@@ -181,8 +181,9 @@ void UDaySequenceModifierComponent::UpdateEditorPreview(float DeltaTime)
 	if (bIsComponentEnabled && bPreview && IsRegistered() && !GetWorld()->IsGameWorld())
 	{
 		const float OldEffectiveBlendWeight = GetBlendWeight();
+		UpdateInternalBlendWeight();
 		
-		if (UpdateInternalBlendWeight() > UE_SMALL_NUMBER)
+		if (GetBlendWeight() > UE_SMALL_NUMBER)
 		{
 			EnableModifier();
 			
@@ -305,9 +306,7 @@ void UDaySequenceModifierComponent::DaySequenceUpdate()
 		// Force the expensive update
 		UpdateInternalBlendWeight();
 		
-		// For the purposes of enable/disable we ignore BlendPolicy and directly use InternalBlendWeight,
-		// but the easing function will respect it (by calling GetBlendWeight).
-		if (InternalBlendWeight > UE_SMALL_NUMBER)
+		if (GetBlendWeight() > UE_SMALL_NUMBER)
 		{
 			EnableModifier();
 		}
