DaySequenceModifierVolume: Adding nullcheck for ModifierComponent.

diff --git a/Engine/Plugins/Experimental/DaySequence/Source/DaySequence/Private/Actors/DaySequenceModifierVolume.cpp b/Engine/Plugins/Experimental/DaySequence/Source/DaySequence/Private/Actors/DaySequenceModifierVolume.cpp
--- a/Engine/Plugins/Experimental/DaySequence/Source/DaySequence/Private/Actors/DaySequenceModifierVolume.cpp
+++ b/Engine/Plugins/Experimental/DaySequence/Source/DaySequence/Private/Actors/DaySequenceModifierVolume.cpp
@@ -167,8 +167,10 @@ void ADaySequenceModifierVolume::BindToDaySequenceActor()
 		DaySequenceModifier->BindToDaySequenceActor(DaySequenceActor);
 		for (auto AdditionalPlayerIterator = AdditionalPlayers.CreateIterator(); AdditionalPlayerIterator; ++AdditionalPlayerIterator)
 		{
-			UDaySequenceModifierComponent* ModifierComponent = AdditionalPlayerIterator->Value;
-			ModifierComponent->BindToDaySequenceActor(DaySequenceActor);
+			if(UDaySequenceModifierComponent* ModifierComponent = AdditionalPlayerIterator->Value)
+			{
+				ModifierComponent->BindToDaySequenceActor(DaySequenceActor);
+			}
 		}
 
 		OnDaySequenceActorBound(DaySequenceActor);
