DaySequence: Added support for blending blend weight over a fixed duration to DaySequenceModifierComponent.

diff --git a/Engine/Plugins/Experimental/DaySequence/Source/DaySequence/Private/DaySequenceModifierComponent.cpp b/Engine/Plugins/Experimental/DaySequence/Source/DaySequence/Private/DaySequenceModifierComponent.cpp
--- a/Engine/Plugins/Experimental/DaySequence/Source/DaySequence/Private/DaySequenceModifierComponent.cpp
+++ b/Engine/Plugins/Experimental/DaySequence/Source/DaySequence/Private/DaySequenceModifierComponent.cpp
@@ -153,6 +153,10 @@ UDaySequenceModifierComponent::UDaySequenceModifierComponent(const FObjectInitia
 	BlendAmount = 100.f;
 	UserBlendWeight = 1.f;
 	InternalBlendWeight = 1.f;
+	BlendTime = 0.f;
+	TimedBlendingLastUpdated = 0.f;
+	TimedBlendingStartWeight = 0.f;
+	TimedBlendingTargetWeight = 0.f;
 #if DAY_SEQUENCE_ENABLE_DRAW_DEBUG
 	DebugLevel = 0;
 #endif
@@ -1195,23 +1199,42 @@ float UDaySequenceModifierComponent::UpdateInternalBlendWeight()
 		}
 		break;
 	}
-
-	if (Mode == EDaySequenceModifierMode::Time)
+	
+	if (BlendTime > 0.0f)
 	{
 		if (const UWorld* World = GetWorld())
 		{
 			const float CurrentTime = World->GetTimeSeconds();
-            const float DeltaTime = CurrentTime - TimedBlendingLastUpdated;
-        	TimedBlendingLastUpdated = CurrentTime;
-            
-        	if (InternalBlendWeight > 0.f)
-        	{
-        		InternalBlendWeight = FMath::Clamp(PreviousBlendWeight + DeltaTime / FMath::Max(BlendTime, UE_SMALL_NUMBER), UE_SMALL_NUMBER, 1.f);
-        	}
-        	else
-        	{
-        		InternalBlendWeight = FMath::Clamp(PreviousBlendWeight - DeltaTime / FMath::Max(BlendTime, UE_SMALL_NUMBER), 0.f, 1.f);
-        	}
+			float DeltaWeightRange;
+			if (Mode == EDaySequenceModifierMode::Time)
+			{
+				TimedBlendingTargetWeight = InternalBlendWeight > 0.f ? 1.f : 0.f;
+				TimedBlendingStartWeight = PreviousBlendWeight;
+				TimedBlendingLastUpdated = CurrentTime - World->GetDeltaSeconds();
+				DeltaWeightRange = InternalBlendWeight > 0.f ? 1.f : -1.f;
+			}
+			else
+			{
+				if (FMath::Abs(InternalBlendWeight - TimedBlendingTargetWeight) > 0.f)
+				{
+					TimedBlendingTargetWeight = InternalBlendWeight;
+					TimedBlendingStartWeight = PreviousBlendWeight;
+					// Ensure at least DeltaSeconds in progress on first tick so that the
+					// blending progresses as players progress through the blend region.
+					TimedBlendingLastUpdated = CurrentTime - World->GetDeltaSeconds();
+				}
+				DeltaWeightRange = TimedBlendingTargetWeight - TimedBlendingStartWeight;
+			}
+			
+			const float DeltaTime = CurrentTime - TimedBlendingLastUpdated;
+			const float DeltaTimeRatio = FMath::Clamp(DeltaTime / FMath::Max(BlendTime, UE_SMALL_NUMBER), 0.f, 1.f);
+			const float DeltaWeight = DeltaWeightRange * DeltaTimeRatio;
+
+			const float MinWeight = DeltaWeight > 0.f ? TimedBlendingStartWeight : TimedBlendingTargetWeight;
+			const float MaxWeight = DeltaWeight > 0.f ? TimedBlendingTargetWeight : TimedBlendingStartWeight;
+			InternalBlendWeight = FMath::Clamp(TimedBlendingStartWeight + DeltaWeight, MinWeight, MaxWeight);
+			
+			UE_LOG(LogDaySequence, VeryVerbose, TEXT("[BlendWeight]: Start = %f, Target = %f, DeltaTime = %f, Weight = %f"), TimedBlendingStartWeight, TimedBlendingTargetWeight, DeltaTime, InternalBlendWeight);
 		}
 		else
 		{
diff --git a/Engine/Plugins/Experimental/DaySequence/Source/DaySequence/Public/DaySequenceModifierComponent.h b/Engine/Plugins/Experimental/DaySequence/Source/DaySequence/Public/DaySequenceModifierComponent.h
--- a/Engine/Plugins/Experimental/DaySequence/Source/DaySequence/Public/DaySequenceModifierComponent.h
+++ b/Engine/Plugins/Experimental/DaySequence/Source/DaySequence/Public/DaySequenceModifierComponent.h
@@ -305,8 +305,8 @@ protected:
 	UPROPERTY(EditAnywhere, BlueprintReadWrite, Category="Day Sequence", meta=(DisplayAfter="Mode", EditCondition="Mode==EDaySequenceModifierMode::Volume", EditConditionHides))
 	float BlendAmount;
 
-	/** Defines the amount of time (in seconds) that it takes for blend weight to move across the full range (0.0, 1.0) when Mode == EDaySequenceModifierMode::Time. */
-	UPROPERTY(EditAnywhere, BlueprintReadWrite, Category="Day Sequence", meta=(DisplayAfter="Mode", EditCondition="Mode==EDaySequenceModifierMode::Time", EditConditionHides))
+	/** Defines the amount of time (in seconds) that it takes for the blend weight to reach the target weight */
+	UPROPERTY(EditAnywhere, BlueprintReadWrite, Category="Day Sequence", meta=(DisplayAfter="Mode"))
 	float BlendTime;
 
 	/** User specified blend weight. The final blend weight is determined by BlendPolicy. */
@@ -390,6 +390,8 @@ private:
 
 	/** Used by UpdateInternalBlendWeight to compute time delta between InternalBlendWeight updates. */
 	float TimedBlendingLastUpdated;
+	float TimedBlendingStartWeight;
+	float TimedBlendingTargetWeight;
 	
 	/**
 	 * Creates and adds or marks for preserve all subsections that the modifier is responsible for.
