DaySequence: Set modifier component as root component of DaySequenceModifierVolume.

diff --git a/Engine/Plugins/Experimental/DaySequence/Source/DaySequence/Private/Actors/DaySequenceModifierVolume.cpp b/Engine/Plugins/Experimental/DaySequence/Source/DaySequence/Private/Actors/DaySequenceModifierVolume.cpp
--- a/Engine/Plugins/Experimental/DaySequence/Source/DaySequence/Private/Actors/DaySequenceModifierVolume.cpp
+++ b/Engine/Plugins/Experimental/DaySequence/Source/DaySequence/Private/Actors/DaySequenceModifierVolume.cpp
@@ -30,7 +30,7 @@ ADaySequenceModifierVolume::ADaySequenceModifierVolume(const FObjectInitializer&
 	PrimaryActorTick.bCanEverTick = true;
 	
 	DaySequenceModifier = CreateDefaultSubobject<UDaySequenceModifierComponent>(TEXT("DaySequenceModifier"));
-	DaySequenceModifier->SetupAttachment(RootComponent);
+	SetRootComponent(DaySequenceModifier);
 
 	DefaultBox = CreateDefaultSubobject<UBoxComponent>(TEXT("Box"));
 	DefaultBox->SetupAttachment(DaySequenceModifier);
