DaySequence: Several debug overlay changes

1) Stop drawing BlueprintEditor world DSA debug info
2) Sort entries by HBias
3) Remove local/remote role info
4) Incorporate actor name into Name column for modifier component entries

diff --git a/Engine/Plugins/Experimental/DaySequence/Source/DaySequence/Private/DaySequenceActor.cpp b/Engine/Plugins/Experimental/DaySequence/Source/DaySequence/Private/DaySequenceActor.cpp
--- a/Engine/Plugins/Experimental/DaySequence/Source/DaySequence/Private/DaySequenceActor.cpp
+++ b/Engine/Plugins/Experimental/DaySequence/Source/DaySequence/Private/DaySequenceActor.cpp
@@ -49,9 +49,10 @@ ADaySequenceActor::FOnSubSectionRemovedEvent ADaySequenceActor::OnSubSectionRemo
 namespace UE::DaySequence
 {
 #if DAY_SEQUENCE_ENABLE_DRAW_DEBUG
-	FDaySequenceDebugEntry::FDaySequenceDebugEntry(FShowDebugDataConditionFunction InShowCondition, FGetDebugDataFunction InGetData)
+	FDaySequenceDebugEntry::FDaySequenceDebugEntry(FShowDebugDataConditionFunction InShowCondition, FGetDebugDataFunction InGetData, FGetNumericValue InGetNumericValue)
 	: ShowCondition(InShowCondition)
 	, GetData(InGetData)
+	, GetNumericValue(InGetNumericValue)
 	{}
 	
 	int32 GDaySequenceDebugLevel = 2;
@@ -928,15 +929,17 @@ UMovieSceneSubSection* ADaySequenceActor::InitializeDaySequence(const FDaySequen
         {
         	if (IsValid(SubSection))
         	{
-        		(*DebugData).FindOrAdd("Actor") = GetFName().ToString();
-        		(*DebugData).FindOrAdd("Local Role") = StaticEnum<ENetRole>()->GetNameStringByValue(GetLocalRole());
-        		(*DebugData).FindOrAdd("Remote Role") = StaticEnum<ENetRole>()->GetNameStringByValue(GetRemoteRole());
+        		(*DebugData).FindOrAdd("Name") = GetPathName(GetOwner() ? GetOwner()->GetOuter() : nullptr);
         		(*DebugData).FindOrAdd("Authority") = HasAuthority() ? "True" : "False";
         		(*DebugData).FindOrAdd("Sequence Name") = SubSection->GetSequence() ? SubSection->GetSequence()->GetFName().ToString() : "None";
         		(*DebugData).FindOrAdd("Mute State") = SubSection->IsActive() ? "Active" : "Muted";
         		(*DebugData).FindOrAdd("Hierarchical Bias") = FString::Printf(TEXT("%d"), SubSection->Parameters.HierarchicalBias);
         	}
         	return DebugData;
+        },
+        [SubSection]()
+        {
+			return IsValid(SubSection) ? SubSection->Parameters.HierarchicalBias : 0;
         }));
 	}
 #endif
@@ -1689,7 +1692,7 @@ void ADaySequenceActor::OnShowDebugInfo(AHUD* HUD, UCanvas* Canvas, const FDebug
 {
 	using namespace UE::DaySequence;
 	
-	if (!Canvas || !GEngine || GetNetMode() == NM_DedicatedServer)
+	if (!Canvas || !GEngine || GetNetMode() == NM_DedicatedServer || !HUD || GetWorld() != HUD->GetWorld())
 	{
 		return;
 	}
@@ -1719,6 +1722,37 @@ void ADaySequenceActor::OnShowDebugInfo(AHUD* HUD, UCanvas* Canvas, const FDebug
 		auto TransformFunction = [](const TWeakPtr<FDaySequenceDebugEntry>& Entry){ return Entry.Pin()->GetData(); };
 		Algo::TransformIf(InWeakArray, OutSharedArray, TransformPredicate, TransformFunction);
 	};
+
+	class FDebugEntriesLessThan
+	{
+	public:
+
+		bool operator()(const TWeakPtr<FDaySequenceDebugEntry>& LHS, const TWeakPtr<FDaySequenceDebugEntry>& RHS) const
+		{
+			const TSharedPtr<FDaySequenceDebugEntry> PinnedLHS = LHS.Pin();
+			const TSharedPtr<FDaySequenceDebugEntry> PinnedRHS = RHS.Pin();
+
+			if (PinnedLHS && PinnedRHS)
+			{
+				return (PinnedLHS->GetNumericValue ? PinnedLHS->GetNumericValue() : 0) < (PinnedRHS->GetNumericValue ? PinnedRHS->GetNumericValue() : 0);
+			}
+
+			// LHS is expired & RHS is valid, expired elements are always less than valid elements.
+			if (PinnedLHS)
+			{
+				return false;
+			}
+
+			// LHS is valid & RHS is expired, expired elements are always less than valid elements.
+			if (PinnedRHS)
+			{
+				return true;
+			}
+
+			// If both are expired, then we consider the values equal (so LHS is not less than RHS).
+			return false;
+		}
+	};
 	
 	TArray<TSharedPtr<TMap<FString, FString>>> EntriesToDraw;
 	for (TPair<FName, TPair<FDebugEntryArray, FDebugCategoryDrawFunction>>& Category : DebugEntries)
@@ -1730,6 +1764,8 @@ void ADaySequenceActor::OnShowDebugInfo(AHUD* HUD, UCanvas* Canvas, const FDebug
 			FDebugEntryArray& Entries = Category.Value.Key;
 			FDebugCategoryDrawFunction& DrawFunction = Category.Value.Value;
 
+			Entries.StableSort(FDebugEntriesLessThan());
+			
 			RemoveStaleEntriesAndGetPinnedArray(Entries, EntriesToDraw);
 
 			DrawFunction(Canvas, EntriesToDraw, Category.Key.ToString());
diff --git a/Engine/Plugins/Experimental/DaySequence/Source/DaySequence/Private/DaySequenceModifierComponent.cpp b/Engine/Plugins/Experimental/DaySequence/Source/DaySequence/Private/DaySequenceModifierComponent.cpp
--- a/Engine/Plugins/Experimental/DaySequence/Source/DaySequence/Private/DaySequenceModifierComponent.cpp
+++ b/Engine/Plugins/Experimental/DaySequence/Source/DaySequence/Private/DaySequenceModifierComponent.cpp
@@ -408,9 +408,7 @@ void UDaySequenceModifierComponent::BindToDaySequenceActor(ADaySequenceActor* Da
 		[this](){ return ShouldShowDebugInfo(); },
 		[this, DebugData]()
 		{
-			(*DebugData).FindOrAdd("Actor") = GetOwner()->GetFName().ToString();
-			(*DebugData).FindOrAdd("Local Role") = StaticEnum<ENetRole>()->GetNameStringByValue(GetOwner()->GetLocalRole());
-			(*DebugData).FindOrAdd("Remote Role") = StaticEnum<ENetRole>()->GetNameStringByValue(GetOwner()->GetRemoteRole());
+			(*DebugData).FindOrAdd("Name") = GetPathName(GetOwner() ? GetOwner()->GetOuter() : nullptr);
 			(*DebugData).FindOrAdd("Component Enabled") = bIsComponentEnabled ? "True" : "False";
 			(*DebugData).FindOrAdd("Modifier Enabled") = bIsEnabled ? "True" : "False";
 			(*DebugData).FindOrAdd("Blend Weight") = FString::Printf(TEXT("%.5f"), GetBlendWeight());
@@ -419,7 +417,8 @@ void UDaySequenceModifierComponent::BindToDaySequenceActor(ADaySequenceActor* Da
 			(*DebugData).FindOrAdd("Blend Target" ) = BlendTarget ? BlendTarget->GetName() : "None";
 
 			return DebugData;
-		});
+		},
+		[]() { return 0; });
 		
 		DaySequenceActor->RegisterDebugEntry(DebugEntry, ShowDebug_ModifierCategory);
 #endif
@@ -940,9 +939,7 @@ UMovieSceneSubSection* UDaySequenceModifierComponent::InitializeDaySequence(cons
 	{
 		if (IsValid(SubSection))
 		{
-			(*DebugData).FindOrAdd("Actor") = GetOwner()->GetFName().ToString();
-			(*DebugData).FindOrAdd("Local Role") = StaticEnum<ENetRole>()->GetNameStringByValue(GetOwner()->GetLocalRole());
-			(*DebugData).FindOrAdd("Remote Role") = StaticEnum<ENetRole>()->GetNameStringByValue(GetOwner()->GetRemoteRole());
+			(*DebugData).FindOrAdd("Name") = GetPathName(GetOwner() ? GetOwner()->GetOuter() : nullptr);
 			(*DebugData).FindOrAdd("Authority") = GetOwner()->HasAuthority() ? "True" : "False";
 			(*DebugData).FindOrAdd("Sequence Name") = SubSection->GetSequence() ? SubSection->GetSequence()->GetFName().ToString() : "None";
 			(*DebugData).FindOrAdd("Mute State") = SubSection->IsActive() ? "Active" : "Muted";
@@ -950,6 +947,10 @@ UMovieSceneSubSection* UDaySequenceModifierComponent::InitializeDaySequence(cons
 		}
 	
 		return DebugData;
+	},
+	[SubSection]()
+	{
+		return IsValid(SubSection) ? SubSection->Parameters.HierarchicalBias : 0;
 	}));
 #endif
 
diff --git a/Engine/Plugins/Experimental/DaySequence/Source/DaySequence/Public/DaySequenceActor.h b/Engine/Plugins/Experimental/DaySequence/Source/DaySequence/Public/DaySequenceActor.h
--- a/Engine/Plugins/Experimental/DaySequence/Source/DaySequence/Public/DaySequenceActor.h
+++ b/Engine/Plugins/Experimental/DaySequence/Source/DaySequence/Public/DaySequenceActor.h
@@ -56,11 +56,13 @@ namespace UE::DaySequence
 	{
 		using FShowDebugDataConditionFunction = TFunction<bool()>;
 		using FGetDebugDataFunction = TFunction<TSharedPtr<TMap<FString, FString>>()>;
+		using FGetNumericValue = TFunction<int32()>;
 
-		FDaySequenceDebugEntry(FShowDebugDataConditionFunction InShowCondition, FGetDebugDataFunction InGetData);
+		FDaySequenceDebugEntry(FShowDebugDataConditionFunction InShowCondition, FGetDebugDataFunction InGetData, FGetNumericValue InGetNumericValue);
 		
 		FShowDebugDataConditionFunction ShowCondition;
 		FGetDebugDataFunction GetData;
+		FGetNumericValue GetNumericValue;
 	};
 	
 	// Alias type which stores an array of weak pointers to debug entries 
