DaySequence: Add support for explicitly disabling hierarchical blending per-subsequence in a day sequence collection

Currently this is only respected by modifier components, the DSA never enables hierarchical blending for its own sequences.

diff --git a/Engine/Plugins/Experimental/DaySequence/Source/DaySequence/Private/DaySequenceModifierComponent.cpp b/Engine/Plugins/Experimental/DaySequence/Source/DaySequence/Private/DaySequenceModifierComponent.cpp
--- a/Engine/Plugins/Experimental/DaySequence/Source/DaySequence/Private/DaySequenceModifierComponent.cpp
+++ b/Engine/Plugins/Experimental/DaySequence/Source/DaySequence/Private/DaySequenceModifierComponent.cpp
@@ -789,7 +789,7 @@ UMovieSceneSubSection* UDaySequenceModifierComponent::InitializeDaySequence(cons
 		return nullptr;
 	}
 	
-	auto CreateSubTrack = [this, MovieScene](UDaySequence* Sequence, int BiasOffset, bool bActivate, bool bBlendHierarchicalBias)
+	auto CreateSubTrack = [this, MovieScene, bEnableBlending = Entry.bEnableHierarchicalBlending](UDaySequence* Sequence, int BiasOffset, bool bActivate)
 	{
 		UDaySequenceTrack*     RootTrack  = MovieScene->AddTrack<UDaySequenceTrack>();
 		RootTrack->ClearFlags(RF_Transactional);
@@ -801,7 +801,7 @@ UMovieSceneSubSection* UDaySequenceModifierComponent::InitializeDaySequence(cons
 		SubSection->Parameters.HierarchicalBias = Bias + BiasOffset;
 		SubSection->Parameters.Flags            = EMovieSceneSubSectionFlags::OverrideRestoreState
 												| (bIgnoreBias ? EMovieSceneSubSectionFlags::IgnoreHierarchicalBias : EMovieSceneSubSectionFlags::None)
-												| (bBlendHierarchicalBias ? EMovieSceneSubSectionFlags::BlendHierarchicalBias : EMovieSceneSubSectionFlags::None);
+												| (bEnableBlending ? EMovieSceneSubSectionFlags::BlendHierarchicalBias : EMovieSceneSubSectionFlags::None);
 
 		const TRange<FFrameNumber> PlaybackRange = MovieScene->GetPlaybackRange();
 		SubSection->SetSequence(Sequence);
@@ -857,8 +857,7 @@ UMovieSceneSubSection* UDaySequenceModifierComponent::InitializeDaySequence(cons
 	};
 	
 	constexpr bool bActivate = true;
-	constexpr bool bBlendHierarchicalBias = true;
-	UMovieSceneSubSection* SubSection = CreateSubTrack(Entry.Sequence, Entry.BiasOffset, bActivate, bBlendHierarchicalBias);
+	UMovieSceneSubSection* SubSection = CreateSubTrack(Entry.Sequence, Entry.BiasOffset, bActivate);
 	
 	if (!SubSections.Contains(SubSection))
 	{
diff --git a/Engine/Plugins/Experimental/DaySequence/Source/DaySequence/Public/DaySequenceCollectionAsset.h b/Engine/Plugins/Experimental/DaySequence/Source/DaySequence/Public/DaySequenceCollectionAsset.h
--- a/Engine/Plugins/Experimental/DaySequence/Source/DaySequence/Public/DaySequenceCollectionAsset.h
+++ b/Engine/Plugins/Experimental/DaySequence/Source/DaySequence/Public/DaySequenceCollectionAsset.h
@@ -37,6 +37,9 @@ struct FDaySequenceCollectionEntry
 	/* The set of conditions which must evaluate to their expected values for this entry to be active. */
 	UPROPERTY(EditAnywhere, Category="Day Sequence")
 	FDaySequenceConditionSet Conditions;
+
+	UPROPERTY(EditAnywhere, Category="Day Sequence", AdvancedDisplay)
+	bool bEnableHierarchicalBlending = true;
 };
 
 UCLASS(MinimalAPI)
