DaySequenceModifierVolume: Back out change from revision 2 which attempted to disable volumes in split screen contexts

diff --git a/Engine/Plugins/Experimental/DaySequence/Source/DaySequence/Private/Actors/DaySequenceModifierVolume.cpp b/Engine/Plugins/Experimental/DaySequence/Source/DaySequence/Private/Actors/DaySequenceModifierVolume.cpp
--- a/Engine/Plugins/Experimental/DaySequence/Source/DaySequence/Private/Actors/DaySequenceModifierVolume.cpp
+++ b/Engine/Plugins/Experimental/DaySequence/Source/DaySequence/Private/Actors/DaySequenceModifierVolume.cpp
@@ -218,7 +218,7 @@ void ADaySequenceModifierVolume::CreatePlayer(APlayerController* InPC)
 	}
 
 	UDaySequenceModifierComponent* PlayerModifier = nullptr;
-	
+
 	if (CachedPlayerController == nullptr)
 	{
 		CachedPlayerController = InPC;
@@ -231,11 +231,11 @@ void ADaySequenceModifierVolume::CreatePlayer(APlayerController* InPC)
 	}
 	else
 	{
-		// If CachedPlayerController is not null (points to the first local player) and we don't support splitscreen, just track the player.
-		AdditionalPlayers.FindOrAdd(InPC) = nullptr;
+		// If CachedPlayerController is not null (points to the first local player) and we don't support splitscreen, just bail out.
+		return;
 	}
 
-	if (PlayerModifier)
+	if (ensureMsgf(PlayerModifier, TEXT("PlayerModifier is nullptr!")))
 	{
 		if (!PlayerModifier->IsRegistered())
 		{
@@ -248,24 +248,6 @@ void ADaySequenceModifierVolume::CreatePlayer(APlayerController* InPC)
 		PlayerModifier->SetBlendTarget(InPC);
 		PlayerModifier->BindToDaySequenceActor(DaySequenceActor);
 	}
-	else if (!bEnableSplitscreenSupport)
-	{
-		// If we get here, we have multiple local players and do not have splitscreen support enabled.
-		// In this case, we will just disable all modifiers.
-		
-		for (TTuple AdditionalPlayer : AdditionalPlayers)
-		{
-			if (AdditionalPlayer.Value)
-			{
-				AdditionalPlayer.Value->DisableComponent();
-			}
-		}
-		
-		if (DaySequenceModifier)
-		{
-			DaySequenceModifier->DisableComponent();
-		}
-	}
 }
 
 bool ADaySequenceModifierVolume::IsSplitscreenSupported() const
