DaySequence: Fixed time blending to account for non zero tick interval setups.
Fixed time blending causing modifiers to turn on at BeginPlay() even when the player does not spawn in the volume.

diff --git a/Engine/Plugins/Experimental/DaySequence/Source/DaySequence/Private/DaySequenceModifierComponent.cpp b/Engine/Plugins/Experimental/DaySequence/Source/DaySequence/Private/DaySequenceModifierComponent.cpp
--- a/Engine/Plugins/Experimental/DaySequence/Source/DaySequence/Private/DaySequenceModifierComponent.cpp
+++ b/Engine/Plugins/Experimental/DaySequence/Source/DaySequence/Private/DaySequenceModifierComponent.cpp
@@ -99,6 +99,13 @@ namespace UE::DaySequence
 		TEXT("When true, day sequence modifier components will automatically disable when they are made invisible."),
 		ECVF_Default
 	);
+
+	static TAutoConsoleVariable<bool> CVarModifierBlendPositionValidateTimestamp(
+		TEXT("DaySequence.Modifier.BlendPositionValidateTimestamp"),
+		true,
+		TEXT("When true, day sequence modifier components will check the camera manager's timestamp for validity before returning a position."),
+		ECVF_Default
+	);
 	
 } // namespace UE::DaySequence
 
@@ -152,7 +159,7 @@ UDaySequenceModifierComponent::UDaySequenceModifierComponent(const FObjectInitia
 	BlendPolicy = EDaySequenceModifierUserBlendPolicy::Minimum;
 	BlendAmount = 100.f;
 	UserBlendWeight = 1.f;
-	InternalBlendWeight = 1.f;
+	InternalBlendWeight = 0.f;
 	BlendTime = 0.f;
 	TimedBlendingLastUpdated = 0.f;
 	TimedBlendingStartWeight = 0.f;
@@ -1084,7 +1091,9 @@ bool UDaySequenceModifierComponent::GetBlendPosition(FVector& InPosition) const
 	}
 	else
 #endif
-	if (const APlayerController* BlendTarget = WeakBlendTarget.Get(); BlendTarget && BlendTarget->PlayerCameraManager)
+	if (const APlayerController* BlendTarget = WeakBlendTarget.Get();
+		BlendTarget && BlendTarget->PlayerCameraManager &&
+		(!UE::DaySequence::CVarModifierBlendPositionValidateTimestamp.GetValueOnAnyThread() || BlendTarget->PlayerCameraManager->GetCameraCacheTime() != 0.f))
 	{
 		CSV_SCOPED_TIMING_STAT(DaySequence, GetPlayerViewPoint);
 		InPosition = BlendTarget->PlayerCameraManager->GetCameraLocation();
@@ -1197,7 +1206,7 @@ float UDaySequenceModifierComponent::UpdateInternalBlendWeight()
 		}
 		else
 		{
-			InternalBlendWeight = 1.f;
+			InternalBlendWeight = 0.f;
 		}
 		break;
 	}
@@ -1208,12 +1217,14 @@ float UDaySequenceModifierComponent::UpdateInternalBlendWeight()
 		{
 			const float CurrentTime = World->GetTimeSeconds();
 			float DeltaWeightRange;
+			float DeltaTime;
 			if (Mode == EDaySequenceModifierMode::Time)
 			{
 				TimedBlendingTargetWeight = InternalBlendWeight > 0.f ? 1.f : 0.f;
 				TimedBlendingStartWeight = PreviousBlendWeight;
-				TimedBlendingLastUpdated = CurrentTime - World->GetDeltaSeconds();
+				DeltaTime = CurrentTime - TimedBlendingLastUpdated;
 				DeltaWeightRange = InternalBlendWeight > 0.f ? 1.f : -1.f;
+				TimedBlendingLastUpdated = CurrentTime;
 			}
 			else
 			{
@@ -1221,14 +1232,17 @@ float UDaySequenceModifierComponent::UpdateInternalBlendWeight()
 				{
 					TimedBlendingTargetWeight = InternalBlendWeight;
 					TimedBlendingStartWeight = PreviousBlendWeight;
+					// TODO: Fix this to account for non-zero tick intervals, where DeltaSeconds
+					// is not representative of a DaySequence update interval.
+					//
 					// Ensure at least DeltaSeconds in progress on first tick so that the
 					// blending progresses as players progress through the blend region.
 					TimedBlendingLastUpdated = CurrentTime - World->GetDeltaSeconds();
 				}
+				DeltaTime = CurrentTime - TimedBlendingLastUpdated;
 				DeltaWeightRange = TimedBlendingTargetWeight - TimedBlendingStartWeight;
 			}
 			
-			const float DeltaTime = CurrentTime - TimedBlendingLastUpdated;
 			const float DeltaTimeRatio = FMath::Clamp(DeltaTime / FMath::Max(BlendTime, UE_SMALL_NUMBER), 0.f, 1.f);
 			const float DeltaWeight = DeltaWeightRange * DeltaTimeRatio;
 
