DaySequence: Disable DaySequenceModifierVolume ticking.

diff --git a/Engine/Plugins/Experimental/DaySequence/Source/DaySequence/Private/Actors/DaySequenceModifierVolume.cpp b/Engine/Plugins/Experimental/DaySequence/Source/DaySequence/Private/Actors/DaySequenceModifierVolume.cpp
--- a/Engine/Plugins/Experimental/DaySequence/Source/DaySequence/Private/Actors/DaySequenceModifierVolume.cpp
+++ b/Engine/Plugins/Experimental/DaySequence/Source/DaySequence/Private/Actors/DaySequenceModifierVolume.cpp
@@ -27,7 +27,7 @@ ADaySequenceModifierVolume::ADaySequenceModifierVolume(const FObjectInitializer&
 : Super(Init)
 , bEnableSplitscreenSupport(false)
 {
-	PrimaryActorTick.bCanEverTick = true;
+	PrimaryActorTick.bCanEverTick = false;
 	
 	DaySequenceModifier = CreateDefaultSubobject<UDaySequenceModifierComponent>(TEXT("DaySequenceModifier"));
 	SetRootComponent(DaySequenceModifier);
