DaySequenceModifierVolume: Disable day sequence modifier volume when split screen is detected

#Synced-CL 48374783

diff --git a/Engine/Plugins/Experimental/DaySequence/Source/DaySequence/Private/Actors/DaySequenceModifierVolume.cpp b/Engine/Plugins/Experimental/DaySequence/Source/DaySequence/Private/Actors/DaySequenceModifierVolume.cpp
--- a/Engine/Plugins/Experimental/DaySequence/Source/DaySequence/Private/Actors/DaySequenceModifierVolume.cpp
+++ b/Engine/Plugins/Experimental/DaySequence/Source/DaySequence/Private/Actors/DaySequenceModifierVolume.cpp
@@ -21,6 +21,13 @@ namespace UE::DaySequence
 	TEXT("When true, Day Sequence Modifier Volumes attempt to initialize transient modifier components for all local players."),
 	ECVF_Default
 	);
+
+	static TAutoConsoleVariable<bool> CVarModifierVolumeCanDisableModifiersDuringSplitscreen(
+	TEXT("DaySequence.ModifierVolume.CanDisableModifiersDuringSplitscreen"),
+	true,
+	TEXT("When true and splitscreen support is not enabled, the modifier volume will disable all managed modifier components."),
+	ECVF_Default
+	);
 }
 
 ADaySequenceModifierVolume::ADaySequenceModifierVolume(const FObjectInitializer& Init)
@@ -97,11 +104,11 @@ void ADaySequenceModifierVolume::Initialize()
 		
 		if (World->IsGameWorld())
 		{
-			auto HandleNewPlayerController = [this](APlayerController* PlayerController)
+			auto HandleNewPlayerController = [](ADaySequenceModifierVolume* DSMV, APlayerController* PlayerController)
 			{
-				if (PlayerController->IsLocalController())
+				if (DSMV && PlayerController->IsLocalController())
 				{
-					CreatePlayer(PlayerController);
+					DSMV->CreatePlayer(PlayerController);
 				}
 			};
 
@@ -110,15 +117,21 @@ void ADaySequenceModifierVolume::Initialize()
 			{
 				if (APlayerController* PlayerController = Iterator->Get())
 				{
-					HandleNewPlayerController(PlayerController);
+					HandleNewPlayerController(this, PlayerController);
 				}
 			}
 
-			ActorSpawnedHandle = World->AddOnActorSpawnedHandler(FOnActorSpawned::FDelegate::CreateWeakLambda(this, [HandleNewPlayerController](AActor* SpawnedActor)
+			ActorSpawnedHandle = World->AddOnActorSpawnedHandler(FOnActorSpawned::FDelegate::CreateWeakLambda(this, [this, HandleNewPlayerController, World](AActor* SpawnedActor)
 			{
 				if (APlayerController* PlayerController = Cast<APlayerController>(SpawnedActor))
 				{
-					HandleNewPlayerController(PlayerController);
+					World->GetTimerManager().SetTimerForNextTick([WeakThis = TWeakObjectPtr(this), HandleNewPlayerController, PlayerController]()
+					{
+						if (TStrongObjectPtr<ADaySequenceModifierVolume> PinnedThis = WeakThis.Pin())
+						{
+							HandleNewPlayerController(PinnedThis.Get(), PlayerController);
+						}
+					});
 				}
 			}));
 		}
@@ -212,7 +225,7 @@ void ADaySequenceModifierVolume::SetupDaySequenceSubsystemCallbacks()
 
 void ADaySequenceModifierVolume::CreatePlayer(APlayerController* InPC)
 {
-	if (CachedPlayerController == InPC || AdditionalPlayers.Find(InPC))
+	if (!InPC || !InPC->GetLocalPlayer() || CachedPlayerController == InPC || AdditionalPlayers.Find(InPC))
 	{
 		return;
 	}
@@ -231,11 +244,11 @@ void ADaySequenceModifierVolume::CreatePlayer(APlayerController* InPC)
 	}
 	else
 	{
-		// If CachedPlayerController is not null (points to the first local player) and we don't support splitscreen, just bail out.
-		return;
+		// If CachedPlayerController is not null (points to the first local player) and we don't support splitscreen, just track the player.
+		AdditionalPlayers.FindOrAdd(InPC) = nullptr;
 	}
 
-	if (ensureMsgf(PlayerModifier, TEXT("PlayerModifier is nullptr!")))
+	if (PlayerModifier)
 	{
 		if (!PlayerModifier->IsRegistered())
 		{
@@ -248,6 +261,24 @@ void ADaySequenceModifierVolume::CreatePlayer(APlayerController* InPC)
 		PlayerModifier->SetBlendTarget(InPC);
 		PlayerModifier->BindToDaySequenceActor(DaySequenceActor);
 	}
+	else if (!IsSplitscreenSupported() && UE::DaySequence::CVarModifierVolumeCanDisableModifiersDuringSplitscreen.GetValueOnAnyThread())
+	{
+		// If we get here, we have multiple local players and do not have splitscreen support enabled and our CVar state allows us to disable the managed modifiers.
+		// In this case, we will just disable all modifiers.
+		
+		for (TTuple AdditionalPlayer : AdditionalPlayers)
+		{
+			if (AdditionalPlayer.Value)
+			{
+				AdditionalPlayer.Value->DisableComponent();
+			}
+		}
+		
+		if (DaySequenceModifier)
+		{
+			DaySequenceModifier->DisableComponent();
+		}
+	}
 }
 
 bool ADaySequenceModifierVolume::IsSplitscreenSupported() const
