DaySequencePlayer: Add HasEverEvaluated to allow interface users to determine if at least 1 evaluation has ever completed

#Synced-CL 48374783

diff --git a/Engine/Plugins/Experimental/DaySequence/Source/DaySequence/Private/DaySequencePlayer.cpp b/Engine/Plugins/Experimental/DaySequence/Source/DaySequence/Private/DaySequencePlayer.cpp
--- a/Engine/Plugins/Experimental/DaySequence/Source/DaySequence/Private/DaySequencePlayer.cpp
+++ b/Engine/Plugins/Experimental/DaySequence/Source/DaySequence/Private/DaySequencePlayer.cpp
@@ -552,6 +552,11 @@ TSharedPtr<UE::DaySequence::FOverrideUpdateIntervalHandle> UDaySequencePlayer::G
 	return MakeShareable(new UE::DaySequence::FOverrideUpdateIntervalHandle(this));
 }
 
+bool UDaySequencePlayer::HasEverEvaluated() const
+{
+	return bHasEverEvaluated;
+}
+
 FQualifiedFrameTime UDaySequencePlayer::GetCurrentTime() const
 {
 	FFrameTime Time = PlayPosition.GetCurrentPosition();
@@ -785,6 +790,7 @@ void UDaySequencePlayer::Initialize(UMovieSceneSequence* InSequence)
 		StopAtCurrentTime();
 	}
 
+	bHasEverEvaluated = false;
 	Sequence = InSequence;
 
 	FFrameTime StartTimeWithOffset = StartTime;
@@ -1197,6 +1203,8 @@ void UDaySequencePlayer::PreEvaluation(const FMovieSceneContext& Context)
 
 void UDaySequencePlayer::PostEvaluation(const FMovieSceneContext& Context)
 {
+	bHasEverEvaluated = true;
+
 	FFrameTime CurrentTime  = ConvertFrameTime(Context.GetTime(),         Context.GetFrameRate(), PlayPosition.GetInputRate());
 	FFrameTime PreviousTime = ConvertFrameTime(Context.GetPreviousTime(), Context.GetFrameRate(), PlayPosition.GetInputRate());
 	OnDaySequencePlayerUpdate.Broadcast(*this, CurrentTime, PreviousTime);
diff --git a/Engine/Plugins/Experimental/DaySequence/Source/DaySequence/Private/DaySequencePlayer.h b/Engine/Plugins/Experimental/DaySequence/Source/DaySequence/Private/DaySequencePlayer.h
--- a/Engine/Plugins/Experimental/DaySequence/Source/DaySequence/Private/DaySequencePlayer.h
+++ b/Engine/Plugins/Experimental/DaySequence/Source/DaySequence/Private/DaySequencePlayer.h
@@ -117,6 +117,8 @@ public:
 	
 public:
 
+	UE_API virtual bool HasEverEvaluated() const override;
+
 	/**
 	 * Get the current playback position
 	 * @return The current playback position
@@ -644,6 +646,8 @@ private:
 	float LastUnpausedTimeSeconds = -1;
 	/** The value of UWorld::GetTimeSeconds last time this player was evaluated */
 	float LastTimeSeconds = -1;
+
+	bool bHasEverEvaluated = false;
 };
 
 #undef UE_API
diff --git a/Engine/Plugins/Experimental/DaySequence/Source/DaySequence/Public/IDaySequencePlayer.h b/Engine/Plugins/Experimental/DaySequence/Source/DaySequence/Public/IDaySequencePlayer.h
--- a/Engine/Plugins/Experimental/DaySequence/Source/DaySequence/Public/IDaySequencePlayer.h
+++ b/Engine/Plugins/Experimental/DaySequence/Source/DaySequence/Public/IDaySequencePlayer.h
@@ -34,6 +34,8 @@ public:
 	
 	virtual void Pause() = 0;
 
+	virtual bool HasEverEvaluated() const = 0;
+
 	virtual FQualifiedFrameTime GetCurrentTime() const = 0;
 	
 	virtual FQualifiedFrameTime GetDuration() const = 0;
