DaySequence: Fixed timed blends for global/volume modifiers to account for DSAs with non-zero budgets or sequence update intervals.

#Synced-CL 48375517

diff --git a/Engine/Plugins/Experimental/DaySequence/Source/DaySequence/Private/DaySequenceModifierComponent.cpp b/Engine/Plugins/Experimental/DaySequence/Source/DaySequence/Private/DaySequenceModifierComponent.cpp
--- a/Engine/Plugins/Experimental/DaySequence/Source/DaySequence/Private/DaySequenceModifierComponent.cpp
+++ b/Engine/Plugins/Experimental/DaySequence/Source/DaySequence/Private/DaySequenceModifierComponent.cpp
@@ -164,6 +164,7 @@ UDaySequenceModifierComponent::UDaySequenceModifierComponent(const FObjectInitia
 	TimedBlendingLastUpdated = 0.f;
 	TimedBlendingStartWeight = 0.f;
 	TimedBlendingTargetWeight = 0.f;
+	TimedBlendingTargetWeightLastUpdated = 0.f;
 #if DAY_SEQUENCE_ENABLE_DRAW_DEBUG
 	DebugLevel = 0;
 #endif
@@ -1224,24 +1225,20 @@ float UDaySequenceModifierComponent::UpdateInternalBlendWeight()
 				TimedBlendingStartWeight = PreviousBlendWeight;
 				DeltaTime = CurrentTime - TimedBlendingLastUpdated;
 				DeltaWeightRange = InternalBlendWeight > 0.f ? 1.f : -1.f;
-				TimedBlendingLastUpdated = CurrentTime;
 			}
-			else
+			else // Global || Volume
 			{
 				if (FMath::Abs(InternalBlendWeight - TimedBlendingTargetWeight) > 0.f)
 				{
 					TimedBlendingTargetWeight = InternalBlendWeight;
 					TimedBlendingStartWeight = PreviousBlendWeight;
-					// TODO: Fix this to account for non-zero tick intervals, where DeltaSeconds
-					// is not representative of a DaySequence update interval.
-					//
-					// Ensure at least DeltaSeconds in progress on first tick so that the
-					// blending progresses as players progress through the blend region.
-					TimedBlendingLastUpdated = CurrentTime - World->GetDeltaSeconds();
+					TimedBlendingTargetWeightLastUpdated = TimedBlendingLastUpdated == 0.f ? CurrentTime : TimedBlendingLastUpdated;
 				}
-				DeltaTime = CurrentTime - TimedBlendingLastUpdated;
+				DeltaTime = CurrentTime - TimedBlendingTargetWeightLastUpdated;
 				DeltaWeightRange = TimedBlendingTargetWeight - TimedBlendingStartWeight;
 			}
+
+			TimedBlendingLastUpdated = CurrentTime;
 			
 			const float DeltaTimeRatio = FMath::Clamp(DeltaTime / FMath::Max(BlendTime, UE_SMALL_NUMBER), 0.f, 1.f);
 			const float DeltaWeight = DeltaWeightRange * DeltaTimeRatio;
diff --git a/Engine/Plugins/Experimental/DaySequence/Source/DaySequence/Public/DaySequenceModifierComponent.h b/Engine/Plugins/Experimental/DaySequence/Source/DaySequence/Public/DaySequenceModifierComponent.h
--- a/Engine/Plugins/Experimental/DaySequence/Source/DaySequence/Public/DaySequenceModifierComponent.h
+++ b/Engine/Plugins/Experimental/DaySequence/Source/DaySequence/Public/DaySequenceModifierComponent.h
@@ -390,8 +390,11 @@ private:
 
 	/** Used by UpdateInternalBlendWeight to compute time delta between InternalBlendWeight updates. */
 	float TimedBlendingLastUpdated;
+	/** Used by UpdateInternalBlendWeight to compute the weight delta for timed blends */
 	float TimedBlendingStartWeight;
 	float TimedBlendingTargetWeight;
+	/** Used by UpdateInternalBlendWeight to compute time delta between TimedBlendingTargetWeight updates for global/volume timed blends. */
+	float TimedBlendingTargetWeightLastUpdated;
 	
 	/**
 	 * Creates and adds or marks for preserve all subsections that the modifier is responsible for.
