DaySequence: Add support for StaticTime channels to compute blended time of day for different use cases.

#Synced-CL 48909380

diff --git a/Engine/Plugins/Experimental/DaySequence/Source/DaySequence/Private/DaySequenceActor.cpp b/Engine/Plugins/Experimental/DaySequence/Source/DaySequence/Private/DaySequenceActor.cpp
--- a/Engine/Plugins/Experimental/DaySequence/Source/DaySequence/Private/DaySequenceActor.cpp
+++ b/Engine/Plugins/Experimental/DaySequence/Source/DaySequence/Private/DaySequenceActor.cpp
@@ -1298,7 +1298,19 @@ float ADaySequenceActor::GetTimeOfDay() const
 
 float ADaySequenceActor::GetApparentTimeOfDay() const
 {
-	float Time = HasStaticTimeOfDay() ? GetStaticTimeOfDay() : GetTimeOfDay();
+	using namespace UE::DaySequence;
+
+	// Query both FixedTime and LocalFixedTime channels for Apparent time of day.
+	constexpr EStaticTimeChannel ApparentChannels = EStaticTimeChannel::FixedTime | EStaticTimeChannel::LocalFixedTime;
+	float Time;
+	if (StaticTimeManager && StaticTimeManager->HasStaticTime(ApparentChannels))
+	{
+		Time = StaticTimeManager->GetStaticTime(GetTimeOfDay(), DayLength.ToHours(), ApparentChannels);
+	}
+	else
+	{
+		Time = GetTimeOfDay();
+	}
 	const bool bUseCurve = bUseInterpCurve && DayInterpCurve && !bForceDisableDayInterpCurve;
 	Time = bUseCurve ? DayInterpCurve->FloatCurve.Eval(Time, Time) : Time;
 	return Time;
@@ -1343,7 +1355,7 @@ bool ADaySequenceActor::HasStaticTimeOfDay() const
 
 float ADaySequenceActor::GetStaticTimeOfDay() const
 {
-	if (StaticTimeManager && HasStaticTimeOfDay())
+	if (HasStaticTimeOfDay())
 	{
 		return StaticTimeManager->GetStaticTime(GetTimeOfDay(), DayLength.ToHours());
 	}
diff --git a/Engine/Plugins/Experimental/DaySequence/Source/DaySequence/Private/DaySequenceModifierComponent.cpp b/Engine/Plugins/Experimental/DaySequence/Source/DaySequence/Private/DaySequenceModifierComponent.cpp
--- a/Engine/Plugins/Experimental/DaySequence/Source/DaySequence/Private/DaySequenceModifierComponent.cpp
+++ b/Engine/Plugins/Experimental/DaySequence/Source/DaySequence/Private/DaySequenceModifierComponent.cpp
@@ -607,6 +607,8 @@ void UDaySequenceModifierComponent::DisableModifier()
 
 void UDaySequenceModifierComponent::SetInitialTimeOfDay()
 {
+	using namespace UE::DaySequence;
+	
 	if (TargetActor)
 	{
 		const bool bHasAuthority = TargetActor->HasAuthority();
@@ -618,14 +620,17 @@ void UDaySequenceModifierComponent::SetInitialTimeOfDay()
 		{
 		case EDayNightCycleMode::FixedTime:
 		case EDayNightCycleMode::RandomFixedTime:
+		case EDayNightCycleMode::LocalFixedTime:
 			{
 				auto WantsStaticTime = [this]() -> bool
 				{
 					return IsValidChecked(this) && bIsEnabled && bIsComponentEnabled &&
-						(DayNightCycle == EDayNightCycleMode::FixedTime || DayNightCycle == EDayNightCycleMode::RandomFixedTime);
+						(DayNightCycle == EDayNightCycleMode::FixedTime ||
+						DayNightCycle == EDayNightCycleMode::RandomFixedTime ||
+						DayNightCycle == EDayNightCycleMode::LocalFixedTime);
 				};
 
-				auto GetStaticTime = [this, RandomTime, WantsStaticTime](UE::DaySequence::FStaticTimeInfo& OutRequest) -> bool
+				auto GetStaticTime = [this, RandomTime, WantsStaticTime](FStaticTimeInfo& OutRequest) -> bool
 				{
 					if (WantsStaticTime())
 					{
@@ -636,12 +641,10 @@ void UDaySequenceModifierComponent::SetInitialTimeOfDay()
 
 					return false;
 				};
-			
-				TargetActor->RegisterStaticTimeContributor({this, Bias, WantsStaticTime, GetStaticTime});
-			}
-			break;
 
-		case EDayNightCycleMode::LocalFixedTime:
+				EStaticTimeChannel Channel = DayNightCycle == EDayNightCycleMode::LocalFixedTime ? EStaticTimeChannel::LocalFixedTime : EStaticTimeChannel::FixedTime;
+				TargetActor->RegisterStaticTimeContributor({this, Bias, WantsStaticTime, GetStaticTime, Channel});
+			}
 			break;
 
 		case EDayNightCycleMode::StartAtSpecifiedTime:
diff --git a/Engine/Plugins/Experimental/DaySequence/Source/DaySequence/Private/DaySequenceStaticTime.cpp b/Engine/Plugins/Experimental/DaySequence/Source/DaySequence/Private/DaySequenceStaticTime.cpp
--- a/Engine/Plugins/Experimental/DaySequence/Source/DaySequence/Private/DaySequenceStaticTime.cpp
+++ b/Engine/Plugins/Experimental/DaySequence/Source/DaySequence/Private/DaySequenceStaticTime.cpp
@@ -56,24 +56,28 @@ namespace UE::DaySequence
 		}
 	}
 	
-	bool FStaticTimeManager::HasStaticTime() const
+	bool FStaticTimeManager::HasStaticTime(EStaticTimeChannel InChannels) const
 	{
 		for (const FStaticTimeContributor& Contributor : Contributors)
 		{
-			if (Contributor.UserObject.Get() && Contributor.WantsStaticTime())
+			if (Contributor.UserObject.Get() &&
+				EnumHasAnyFlags(InChannels, Contributor.Channels) &&
+				Contributor.WantsStaticTime())
 			{
 				return true;
 			}
 		}
 
 		// Reset the blend data when static time goes away.
-		ResetBlendState();
+		ResetBlendState(InChannels);
 
 		return false;
 	}
 	
-	float FStaticTimeManager::GetStaticTime(float InitialTime, float DayLength) const
+	float FStaticTimeManager::GetStaticTime(float InitialTime, float DayLength, const EStaticTimeChannel InChannels) const
 	{
+		FStaticTimeChannelData& Channel = Channels.FindOrAdd(InChannels); 
+		
 		float AccumulatedWeight = 0.f;
 		float AccumulatedTime = 0.f;
 		
@@ -86,7 +90,7 @@ namespace UE::DaySequence
 			const int32 PriorityGroupEndIdx = PriorityGroupStartIdx + GroupSize - 1;
 			
 			// Process this group 
-			const FStaticTimeInfo GroupInfo = ProcessPriorityGroup(PriorityGroupStartIdx, PriorityGroupEndIdx);
+			const FStaticTimeInfo GroupInfo = ProcessPriorityGroup(InChannels, PriorityGroupStartIdx, PriorityGroupEndIdx);
 
 			const float EffectiveGroupWeight = (1.f - AccumulatedWeight) * GroupInfo.BlendWeight;
 			const float EffectiveGroupTime = EffectiveGroupWeight * GroupInfo.StaticTime;
@@ -191,11 +195,11 @@ namespace UE::DaySequence
 
 		// Compute LastBlendOffset based on the number of times (incl. direction) BlendDelta
 		// has crossed the TargetTime +/- HalfDayLength mark.
-		if (LastBlendDelta.IsSet() &&
-			LastBlendDirection.IsSet() &&
-			LastBlendDirection.GetValue() != BlendDirection)
+		if (Channel.LastBlendDelta.IsSet() &&
+			Channel.LastBlendDirection.IsSet() &&
+			Channel.LastBlendDirection.GetValue() != BlendDirection)
 		{
-			const float BlendDeltaChange = BlendDelta - LastBlendDelta.GetValue();
+			const float BlendDeltaChange = BlendDelta - Channel.LastBlendDelta.GetValue();
 			if (FMath::Abs(BlendDeltaChange) > HalfDayLength)
 			{
 				// Crossed TargetTime +- HalfDayLength
@@ -205,17 +209,17 @@ namespace UE::DaySequence
 				const int BlendDeltaSign = BlendDeltaChange < 0.f ? 1 : -1;
 				
 				// Record our winding counts based on direction of the crossover.
-				LastBlendWinding += BlendDeltaSign;
-				LastBlendOffset = FMath::WrapExclusive(LastBlendOffset + BlendDeltaSign, -1, 2);
+				Channel.LastBlendWinding += BlendDeltaSign;
+				Channel.LastBlendOffset = FMath::WrapExclusive(Channel.LastBlendOffset + BlendDeltaSign, -1, 2);
 			}
 			else
 			{
 				// Crossed TargetTime
-				if (LastBlendOffset != 0)
+				if (Channel.LastBlendOffset != 0)
 				{
 					// Loop TargetTime + HalfDayLength  <--> TargetTime - HalfDayLength
 					const int BlendDeltaSign = BlendDeltaChange > 0.f ? 1 : -1;
-					LastBlendOffset = FMath::WrapExclusive(LastBlendOffset + BlendDeltaSign, -1, 2);
+					Channel.LastBlendOffset = FMath::WrapExclusive(Channel.LastBlendOffset + BlendDeltaSign, -1, 2);
 				}
 			}
 		}
@@ -223,7 +227,7 @@ namespace UE::DaySequence
 		// Blend against initial value if necessary
 		if (AccumulatedWeight < 1.0f)
 		{
-			float BlendOffset = LastBlendOffset * DayLength;
+			float BlendOffset = Channel.LastBlendOffset * DayLength;
 			const float StartTime = TargetTime + BlendDelta + BlendOffset;
 
 			// Wrap is inclusive of the upper bound. Ensure we wrap around correctly.
@@ -235,33 +239,34 @@ namespace UE::DaySequence
 			}
 
 #if !NO_LOGGING
-			UE_LOG(LogDaySequence, VeryVerbose, TEXT("BlendDelta: %f | BlendDir: %d | BlendWinding: %d | BlendOffset: %f | StartTime: %f | InitialTime: %f | TargetTime: %f"), BlendDelta, BlendDirection, LastBlendWinding, BlendOffset, StartTime, InitialTime, TargetTime);
+			UE_LOG(LogDaySequence, VeryVerbose, TEXT("BlendDelta: %f | BlendDir: %d | BlendWinding: %d | BlendOffset: %f | StartTime: %f | InitialTime: %f | TargetTime: %f"), BlendDelta, BlendDirection, Channel.LastBlendWinding, BlendOffset, StartTime, InitialTime, TargetTime);
 			UE_LOG(LogDaySequence, VeryVerbose, TEXT("AccumTime: %f | AccumWeight: %f"), AccumulatedTime, AccumulatedWeight);
 #endif
 		}
 		else
 		{
-			ResetBlendState();
+			ResetBlendState(InChannels);
 		}
 		
 		// we could probably ensure(AccumulatedWeight is nearly equal to 1.f), leaving out for now
 
-		LastBlendDelta = BlendDelta;
-		LastBlendDirection = BlendDirection;
+		Channel.LastBlendDelta = BlendDelta;
+		Channel.LastBlendDirection = BlendDirection;
 		
 		return AccumulatedTime;
 	}
 
-	void FStaticTimeManager::ResetBlendState() const
+	void FStaticTimeManager::ResetBlendState(const EStaticTimeChannel InChannels) const
 	{
-		LastBlendDirection.Reset();
-		LastBlendDelta.Reset();
-		LastBlendWinding = 0;
-		LastBlendOffset = 0;
+		FStaticTimeChannelData& Channel = Channels.FindOrAdd(InChannels);
+		Channel.LastBlendDirection.Reset();
+		Channel.LastBlendDelta.Reset();
+		Channel.LastBlendWinding = 0;
+		Channel.LastBlendOffset = 0;
 	}
 
 
-	FStaticTimeInfo FStaticTimeManager::ProcessPriorityGroup(int32 StartIdx, int32 EndIdx) const
+	FStaticTimeInfo FStaticTimeManager::ProcessPriorityGroup(const EStaticTimeChannel InChannels, int32 StartIdx, int32 EndIdx) const
 	{
 		FStaticTimeInfo GroupInfo = {0.f, 0.f};
 		int32 GroupContributors = 0;
@@ -270,7 +275,9 @@ namespace UE::DaySequence
 		{
 			const FStaticTimeContributor& CurrentContributor = Contributors[CurrentIdx];
 
-			if (!CurrentContributor.UserObject.Get() || !CurrentContributor.WantsStaticTime())
+			if (!CurrentContributor.UserObject.Get() ||
+				!EnumHasAnyFlags(InChannels, CurrentContributor.Channels) ||
+				!CurrentContributor.WantsStaticTime())
 			{
 				continue;
 			}
diff --git a/Engine/Plugins/Experimental/DaySequence/Source/DaySequence/Private/DaySequenceStaticTime.h b/Engine/Plugins/Experimental/DaySequence/Source/DaySequence/Private/DaySequenceStaticTime.h
--- a/Engine/Plugins/Experimental/DaySequence/Source/DaySequence/Private/DaySequenceStaticTime.h
+++ b/Engine/Plugins/Experimental/DaySequence/Source/DaySequence/Private/DaySequenceStaticTime.h
@@ -21,10 +21,17 @@ namespace UE::DaySequence
 	// The function signature for the static time requested callback. Must return true for a contributor to be considered
 	using FWantsStaticTimeFunction = TFunction<bool()>;
 	
-	// The function signature for the static time info getter callback. Only called if IsStaticTimeRequested returns true.
-	// This should return the same value that calling FStaticTimeRequestedFunction would return.
+	// The function signature for the static time info getter callback. Only called if WantsStaticTime returns true.
 	using FGetStaticTimeFunction = TFunction<bool(FStaticTimeInfo& OutRequest)>;
-	
+
+	enum class EStaticTimeChannel
+	{
+		None			= 0,
+		FixedTime		= 1 << 0,
+		LocalFixedTime	= 1 << 1
+	};
+	ENUM_CLASS_FLAGS(EStaticTimeChannel);
+
 	// Contributors register an instance of this struct in order to request a static time.
 	struct FStaticTimeContributor
 	{
@@ -39,31 +46,40 @@ namespace UE::DaySequence
 
 		// Provides caller with desired static time information.
 		FGetStaticTimeFunction GetStaticTime;
+
+		// The channel(s) that this static time contributor affects.
+		EStaticTimeChannel Channels = EStaticTimeChannel::FixedTime;
+	};
+
+	struct FStaticTimeChannelData
+	{
+		// Cached blend data
+		int LastBlendWinding = 0;
+		int LastBlendOffset = 0;
+		TOptional<float> LastBlendDelta;
+		TOptional<int> LastBlendDirection;
 	};
 	
 	struct FStaticTimeManager
 	{
 		void AddStaticTimeContributor(const FStaticTimeContributor& NewContributor);
 		void RemoveStaticTimeContributor(const UObject* UserObject);
-		bool HasStaticTime() const;
-		float GetStaticTime(float InitialTime, float DayLength) const;
+		bool HasStaticTime(EStaticTimeChannel InChannels = EStaticTimeChannel::FixedTime) const;
+		float GetStaticTime(float InitialTime, float DayLength, EStaticTimeChannel InChannels = EStaticTimeChannel::FixedTime) const;
 		
 	private:
-		void ResetBlendState() const;
+		void ResetBlendState(EStaticTimeChannel InChannels) const;
 		
 		// Handles multiple contributors of the same priority by averaging their weights and times and returning a single desired weight and time
-		FStaticTimeInfo ProcessPriorityGroup(int32 StartIdx, int32 EndIdx) const;
-		
+		FStaticTimeInfo ProcessPriorityGroup(EStaticTimeChannel InChannels, int32 StartIdx, int32 EndIdx) const;
+
 		TArray<FStaticTimeContributor> Contributors;
 
 		// We precompute this when mutating the Contributors array to avoid redundant reiteration over Contributors when computing static time
 		TMap<int32, int32> PriorityGroupSizes;
 
-		// Cached blend data. Mutable because it is set in (Has/Get)StaticTime.
-		mutable int LastBlendWinding = 0;
-		mutable int LastBlendOffset = 0;
-		mutable TOptional<float> LastBlendDelta;
-		mutable TOptional<int> LastBlendDirection;
+		// Mutable because cached data is set in (Has/Get)StaticTime.
+		mutable TMap<EStaticTimeChannel, FStaticTimeChannelData> Channels;
 	};
 }
 
