DaySequence: Update GetStaticTimeOfDay to return GetTimeOfDay when there are no active static time contributors.

#Synced-CL 48909380

diff --git a/Engine/Plugins/Experimental/DaySequence/Source/DaySequence/Private/DaySequenceActor.cpp b/Engine/Plugins/Experimental/DaySequence/Source/DaySequence/Private/DaySequenceActor.cpp
--- a/Engine/Plugins/Experimental/DaySequence/Source/DaySequence/Private/DaySequenceActor.cpp
+++ b/Engine/Plugins/Experimental/DaySequence/Source/DaySequence/Private/DaySequenceActor.cpp
@@ -34,8 +34,6 @@
 
 #include "Engine/World.h"
 
-#include <limits>
-
 #if WITH_EDITOR
 #include "PropertyHandle.h"
 #endif
@@ -1311,9 +1309,7 @@ float ADaySequenceActor::GetApparentTimeOfDay() const
 	{
 		Time = GetTimeOfDay();
 	}
-	const bool bUseCurve = bUseInterpCurve && DayInterpCurve && !bForceDisableDayInterpCurve;
-	Time = bUseCurve ? DayInterpCurve->FloatCurve.Eval(Time, Time) : Time;
-	return Time;
+	return GetInterpCurveTimeOfDay(Time);
 }
 
 bool ADaySequenceActor::SetTimeOfDay(float InHours)
@@ -1355,12 +1351,18 @@ bool ADaySequenceActor::HasStaticTimeOfDay() const
 
 float ADaySequenceActor::GetStaticTimeOfDay() const
 {
+	using namespace UE::DaySequence;
+	
+	float Time;
 	if (HasStaticTimeOfDay())
 	{
-		return StaticTimeManager->GetStaticTime(GetTimeOfDay(), DayLength.ToHours());
+		Time = StaticTimeManager->GetStaticTime(GetTimeOfDay(), DayLength.ToHours(), EStaticTimeChannel::FixedTime);
 	}
-	
-	return std::numeric_limits<float>::lowest();
+	else
+	{
+		Time = GetTimeOfDay();
+	}
+	return Time;
 }
 
 void ADaySequenceActor::RegisterStaticTimeContributor(const UE::DaySequence::FStaticTimeContributor& NewContributor) const
@@ -1383,6 +1385,12 @@ void ADaySequenceActor::UnregisterStaticTimeContributor(const UObject* InUserObj
 	StaticTimeManager->RemoveStaticTimeContributor(InUserObject);
 }
 
+float ADaySequenceActor::GetInterpCurveTimeOfDay(float Time) const
+{
+	const bool bUseCurve = bUseInterpCurve && DayInterpCurve && !bForceDisableDayInterpCurve;
+	return bUseCurve ? DayInterpCurve->FloatCurve.Eval(Time, Time) : Time;
+}
+
 void ADaySequenceActor::RegisterBindingResolveFunction(FMovieSceneSequenceID SequenceID, FGuid Guid, TFunction<bool(TArray<UObject*, TInlineAllocator<1>>&)> InFunction)
 {
 	UnregisterBindingResolveFunction(SequenceID, Guid);
diff --git a/Engine/Plugins/Experimental/DaySequence/Source/DaySequence/Public/DaySequenceActor.h b/Engine/Plugins/Experimental/DaySequence/Source/DaySequence/Public/DaySequenceActor.h
--- a/Engine/Plugins/Experimental/DaySequence/Source/DaySequence/Public/DaySequenceActor.h
+++ b/Engine/Plugins/Experimental/DaySequence/Source/DaySequence/Public/DaySequenceActor.h
@@ -290,7 +290,7 @@ public:
 	UE_API bool HasStaticTimeOfDay() const;
 
 	/**
-	 * Get this day sequence actor's static (fixed) time-of-day, or numeric_limits<float>::lowest() if it doesn't have one
+	 * Get this day sequence actor's static (fixed) time-of-day, or GetTimeOfDay() if it doesn't have one
 	 */
 	UFUNCTION(BlueprintCallable, Category="TimeOfDay")
 	UE_API float GetStaticTimeOfDay() const;
@@ -473,6 +473,8 @@ protected:
 
 	UE_API void OnSequencePlayerUpdate(const UDaySequencePlayer& Player, FFrameTime CurrentTime, FFrameTime PreviousTime);
 	UE_API virtual void SequencePlayerUpdated(float CurrentTime, float PreviousTime);
+
+	UE_API float GetInterpCurveTimeOfDay(float Time) const;
 	
 public:
 	/**
