DaySequence: Added GetTargetApparentTimeOfDay to fetch the target static time for active static time contributors.

#Synced-CL 49076111

diff --git a/Engine/Plugins/Experimental/DaySequence/Source/DaySequence/Private/DaySequenceActor.cpp b/Engine/Plugins/Experimental/DaySequence/Source/DaySequence/Private/DaySequenceActor.cpp
--- a/Engine/Plugins/Experimental/DaySequence/Source/DaySequence/Private/DaySequenceActor.cpp
+++ b/Engine/Plugins/Experimental/DaySequence/Source/DaySequence/Private/DaySequenceActor.cpp
@@ -1312,6 +1312,23 @@ float ADaySequenceActor::GetApparentTimeOfDay() const
 	return GetInterpCurveTimeOfDay(Time);
 }
 
+float ADaySequenceActor::GetTargetApparentTimeOfDay() const
+{
+	using namespace UE::DaySequence;
+	
+	constexpr EStaticTimeChannel ApparentChannels = EStaticTimeChannel::FixedTime | EStaticTimeChannel::LocalFixedTime;
+	float Time;
+	if (StaticTimeManager && StaticTimeManager->HasStaticTime(ApparentChannels))
+	{
+		Time = StaticTimeManager->GetTargetStaticTime(GetTimeOfDay(), ApparentChannels);
+	}
+	else
+	{
+		Time = GetTimeOfDay();
+	}
+	return Time;
+}
+
 bool ADaySequenceActor::SetTimeOfDay(float InHours)
 {
 #if !UE_BUILD_SHIPPING
diff --git a/Engine/Plugins/Experimental/DaySequence/Source/DaySequence/Private/DaySequenceStaticTime.cpp b/Engine/Plugins/Experimental/DaySequence/Source/DaySequence/Private/DaySequenceStaticTime.cpp
--- a/Engine/Plugins/Experimental/DaySequence/Source/DaySequence/Private/DaySequenceStaticTime.cpp
+++ b/Engine/Plugins/Experimental/DaySequence/Source/DaySequence/Private/DaySequenceStaticTime.cpp
@@ -256,6 +256,32 @@ namespace UE::DaySequence
 		return AccumulatedTime;
 	}
 
+	float FStaticTimeManager::GetTargetStaticTime(const float InitialTime, EStaticTimeChannel InChannels) const
+	{
+		// Process batches of contributors based on priority
+		for (int32 PriorityGroupStartIdx = 0; PriorityGroupStartIdx < Contributors.Num();)
+		{
+			// Compute the index of the final element in this priority group
+			const int32 CurrentPriority = Contributors[PriorityGroupStartIdx].Priority;
+			const int32 GroupSize = PriorityGroupSizes.FindChecked(CurrentPriority);	// checked because if this isn't valid something is very wrong
+			const int32 PriorityGroupEndIdx = PriorityGroupStartIdx + GroupSize - 1;
+			
+			// Process this group 
+			const FStaticTimeInfo GroupInfo = ProcessPriorityGroup(InChannels, PriorityGroupStartIdx, PriorityGroupEndIdx, 1.0f);
+
+			// Return the highest priority group that contains an active contributor
+			if (GroupInfo.BlendWeight > 0.f)
+			{
+				return GroupInfo.StaticTime;
+			}
+
+			// Advance to next priority group
+			PriorityGroupStartIdx = PriorityGroupEndIdx + 1;
+		}
+
+		return InitialTime;
+	}
+
 	void FStaticTimeManager::ResetBlendState(const EStaticTimeChannel InChannels) const
 	{
 		FStaticTimeChannelData& Channel = Channels.FindOrAdd(InChannels);
@@ -266,7 +292,7 @@ namespace UE::DaySequence
 	}
 
 
-	FStaticTimeInfo FStaticTimeManager::ProcessPriorityGroup(const EStaticTimeChannel InChannels, int32 StartIdx, int32 EndIdx) const
+	FStaticTimeInfo FStaticTimeManager::ProcessPriorityGroup(const EStaticTimeChannel InChannels, int32 StartIdx, int32 EndIdx, TOptional<float> OverrideWeight) const
 	{
 		FStaticTimeInfo GroupInfo = {0.f, 0.f};
 		int32 GroupContributors = 0;
@@ -290,7 +316,7 @@ namespace UE::DaySequence
 			CurrentContributor.GetStaticTime(ContributorInfo);
 
 			// Accumulate contributor info (we will later divide these sums once we know how many contributors there are)
-			GroupInfo.BlendWeight += ContributorInfo.BlendWeight;
+			GroupInfo.BlendWeight += OverrideWeight.IsSet() ? OverrideWeight.GetValue() : ContributorInfo.BlendWeight;
 			GroupInfo.StaticTime += ContributorInfo.StaticTime;
 		}
 
diff --git a/Engine/Plugins/Experimental/DaySequence/Source/DaySequence/Private/DaySequenceStaticTime.h b/Engine/Plugins/Experimental/DaySequence/Source/DaySequence/Private/DaySequenceStaticTime.h
--- a/Engine/Plugins/Experimental/DaySequence/Source/DaySequence/Private/DaySequenceStaticTime.h
+++ b/Engine/Plugins/Experimental/DaySequence/Source/DaySequence/Private/DaySequenceStaticTime.h
@@ -14,8 +14,8 @@ namespace UE::DaySequence
 	// The fundamental piece of data used by the static time system.
 	struct FStaticTimeInfo
 	{
-		float BlendWeight;
-		float StaticTime;
+		float BlendWeight = 0.f;
+		float StaticTime = 0.f;
 	};
 
 	// The function signature for the static time requested callback. Must return true for a contributor to be considered
@@ -66,12 +66,13 @@ namespace UE::DaySequence
 		void RemoveStaticTimeContributor(const UObject* UserObject);
 		bool HasStaticTime(EStaticTimeChannel InChannels = EStaticTimeChannel::FixedTime) const;
 		float GetStaticTime(float InitialTime, float DayLength, EStaticTimeChannel InChannels = EStaticTimeChannel::FixedTime) const;
+		float GetTargetStaticTime(float InitialTime, EStaticTimeChannel InChannels) const;
 		
 	private:
 		void ResetBlendState(EStaticTimeChannel InChannels) const;
 		
 		// Handles multiple contributors of the same priority by averaging their weights and times and returning a single desired weight and time
-		FStaticTimeInfo ProcessPriorityGroup(EStaticTimeChannel InChannels, int32 StartIdx, int32 EndIdx) const;
+		FStaticTimeInfo ProcessPriorityGroup(EStaticTimeChannel InChannels, int32 StartIdx, int32 EndIdx, TOptional<float> OverrideWeight = TOptional<float>()) const;
 
 		TArray<FStaticTimeContributor> Contributors;
 
diff --git a/Engine/Plugins/Experimental/DaySequence/Source/DaySequence/Public/DaySequenceActor.h b/Engine/Plugins/Experimental/DaySequence/Source/DaySequence/Public/DaySequenceActor.h
--- a/Engine/Plugins/Experimental/DaySequence/Source/DaySequence/Public/DaySequenceActor.h
+++ b/Engine/Plugins/Experimental/DaySequence/Source/DaySequence/Public/DaySequenceActor.h
@@ -259,6 +259,16 @@ public:
 	UFUNCTION(BlueprintCallable, Category="TimeOfDay")
 	UE_API float GetApparentTimeOfDay() const;
 
+	/**
+	 * Get the current target apparent time of day in hours.
+	 * This is distinct from apparent time of day in that it is only capturing the target time of static time contributors.
+	 * The target time is determined by the highest priority static time contributor independent of blend weight. Multiple
+	 * contributors at the same priority will return the average of their target times.
+	 * 
+	 * @return float, time in hours
+	 */
+	UE_API float GetTargetApparentTimeOfDay() const;
+
 	/**
 	 * Set the current time of day in hours (server only).
 	 *
