DaySequence: Added OnPostEnableModifier to DaySequenceModifierVolume.

#Synced-CL 49076111

diff --git a/Engine/Plugins/Experimental/DaySequence/Source/DaySequence/Private/Actors/DaySequenceModifierVolume.cpp b/Engine/Plugins/Experimental/DaySequence/Source/DaySequence/Private/Actors/DaySequenceModifierVolume.cpp
--- a/Engine/Plugins/Experimental/DaySequence/Source/DaySequence/Private/Actors/DaySequenceModifierVolume.cpp
+++ b/Engine/Plugins/Experimental/DaySequence/Source/DaySequence/Private/Actors/DaySequenceModifierVolume.cpp
@@ -87,6 +87,11 @@ void ADaySequenceModifierVolume::Initialize()
 		return;
 	}
 
+	if (DaySequenceModifier)
+	{
+		DaySequenceModifier->GetOnPostEnableModifier().AddUniqueDynamic(this, &ADaySequenceModifierVolume::OnPostEnableModifier);
+	}
+
 	if (const UWorld* World = GetWorld())
 	{
 #if WITH_EDITOR
@@ -165,6 +170,11 @@ void ADaySequenceModifierVolume::Deinitialize()
 			ReplayScrubbedHandle.Reset();
 		}
 	}
+
+	if (DaySequenceModifier)
+	{
+		DaySequenceModifier->GetOnPostEnableModifier().RemoveAll(this);
+	}
 }
 
 void ADaySequenceModifierVolume::DaySequenceActorSetup()
diff --git a/Engine/Plugins/Experimental/DaySequence/Source/DaySequence/Private/DaySequenceModifierComponent.cpp b/Engine/Plugins/Experimental/DaySequence/Source/DaySequence/Private/DaySequenceModifierComponent.cpp
--- a/Engine/Plugins/Experimental/DaySequence/Source/DaySequence/Private/DaySequenceModifierComponent.cpp
+++ b/Engine/Plugins/Experimental/DaySequence/Source/DaySequence/Private/DaySequenceModifierComponent.cpp
@@ -1281,6 +1281,11 @@ void UDaySequenceModifierComponent::AddVolumeShapeComponent(const FComponentRefe
 	bCachedExternalShapesInvalid = true;
 }
 
+FOnPostEnableModifier& UDaySequenceModifierComponent::GetOnPostEnableModifier()
+{
+	return OnPostEnableModifier;
+}
+
 void UDaySequenceModifierComponent::InvalidateMuteStates() const
 {
 	OnInvalidateMuteStates.Broadcast();
diff --git a/Engine/Plugins/Experimental/DaySequence/Source/DaySequence/Public/Actors/DaySequenceModifierVolume.h b/Engine/Plugins/Experimental/DaySequence/Source/DaySequence/Public/Actors/DaySequenceModifierVolume.h
--- a/Engine/Plugins/Experimental/DaySequence/Source/DaySequence/Public/Actors/DaySequenceModifierVolume.h
+++ b/Engine/Plugins/Experimental/DaySequence/Source/DaySequence/Public/Actors/DaySequenceModifierVolume.h
@@ -25,6 +25,9 @@ public:
 
 	UFUNCTION(BlueprintImplementableEvent)
 	UE_API void OnDaySequenceActorBound(ADaySequenceActor* InActor);
+
+	UFUNCTION(BlueprintImplementableEvent, CallInEditor)
+	UE_API void OnPostEnableModifier();
 	
 protected:
 	//~ Begin AActor interface
diff --git a/Engine/Plugins/Experimental/DaySequence/Source/DaySequence/Public/DaySequenceModifierComponent.h b/Engine/Plugins/Experimental/DaySequence/Source/DaySequence/Public/DaySequenceModifierComponent.h
--- a/Engine/Plugins/Experimental/DaySequence/Source/DaySequence/Public/DaySequenceModifierComponent.h
+++ b/Engine/Plugins/Experimental/DaySequence/Source/DaySequence/Public/DaySequenceModifierComponent.h
@@ -215,6 +215,8 @@ public:
 	void EmptyVolumeShapeComponents();
 	void AddVolumeShapeComponent(const FComponentReference& InShapeReference);
 
+	FOnPostEnableModifier& GetOnPostEnableModifier();
+
 	/*~ Begin UObject interface */
 	virtual void PostLoad() override;
 	/*~ End UObject interface */
