Sequencer: Fixed DaySequence blueprint having invalid characters in its name

#Synced-CL 49157442

diff --git a/Engine/Plugins/Experimental/DaySequence/Source/DaySequence/Private/DaySequence.cpp b/Engine/Plugins/Experimental/DaySequence/Source/DaySequence/Private/DaySequence.cpp
--- a/Engine/Plugins/Experimental/DaySequence/Source/DaySequence/Private/DaySequence.cpp
+++ b/Engine/Plugins/Experimental/DaySequence/Source/DaySequence/Private/DaySequence.cpp
@@ -458,7 +458,7 @@ UBlueprint* UDaySequence::GetDirectorBlueprint() const
 
 FString UDaySequence::GetDirectorBlueprintName() const
 {
-	return GetDisplayName().ToString() + " (Director BP)";
+	return GetDisplayName().ToString() + "_DirectorBP";
 }
 
 void UDaySequence::SetDirectorBlueprint(UBlueprint* NewDirectorBlueprint)
