DaySequenceModifierVolume: Add support for detecting & handling player controller replacement

Instead of tracking APlayerControllers, we just track UPlayers directly. We can fetch an APlayerController from a UPlayer. There are scenarios where a UPlayer might be assigned a new player controller during a game. If we just track APlayerController, we currently interpret this as an entirely new player. Now, we can detect player controller replacement by comparing the new player controller's associated UPlayer to all of our previously encountered UPlayers. If we have seen this UPlayer before, we just reassign the associated modifier component's blend target to the new player controller.

#Synced-CL 49961239

diff --git a/Engine/Plugins/Experimental/DaySequence/Source/DaySequence/Private/Actors/DaySequenceModifierVolume.cpp b/Engine/Plugins/Experimental/DaySequence/Source/DaySequence/Private/Actors/DaySequenceModifierVolume.cpp
--- a/Engine/Plugins/Experimental/DaySequence/Source/DaySequence/Private/Actors/DaySequenceModifierVolume.cpp
+++ b/Engine/Plugins/Experimental/DaySequence/Source/DaySequence/Private/Actors/DaySequenceModifierVolume.cpp
@@ -28,6 +28,13 @@ namespace UE::DaySequence
 	TEXT("When true and splitscreen support is not enabled, the modifier volume will disable all managed modifier components."),
 	ECVF_Default
 	);
+
+	static TAutoConsoleVariable<bool> CVarModifierVolumeCanReplaceBlendTargets(
+	TEXT("DaySequence.ModifierVolume.CanReplaceBlendTargets"),
+	true,
+	TEXT("When true, the modifier will attempt to detect when a local player is assigned a new player controller and will use the new controller for the associated modifier's blend target."),
+	ECVF_Default
+	);
 }
 
 ADaySequenceModifierVolume::ADaySequenceModifierVolume(const FObjectInitializer& Init)
@@ -235,27 +242,47 @@ void ADaySequenceModifierVolume::SetupDaySequenceSubsystemCallbacks()
 
 void ADaySequenceModifierVolume::CreatePlayer(APlayerController* InPC)
 {
-	if (!InPC || !InPC->GetLocalPlayer() || CachedPlayerController == InPC || AdditionalPlayers.Find(InPC))
+	if (!InPC || !InPC->GetLocalPlayer())
 	{
 		return;
 	}
 
+	if (UE::DaySequence::CVarModifierVolumeCanReplaceBlendTargets.GetValueOnAnyThread())
+	{
+		// See if the primary player controller is being replaced.
+		if (CachedPlayer == InPC->Player)
+		{
+			DaySequenceModifier->SetBlendTarget(InPC);
+			return;
+		}
+
+		// See if any additional player controllers are being replaced.
+		for (TPair<TObjectPtr<UPlayer>, TObjectPtr<UDaySequenceModifierComponent>> AdditionalPlayer : AdditionalPlayers)
+		{
+			if (AdditionalPlayer.Key == InPC->Player)
+			{
+				if (AdditionalPlayer.Value) { AdditionalPlayer.Value->SetBlendTarget(InPC); }
+				return;
+			}
+		}
+	}
+
 	UDaySequenceModifierComponent* PlayerModifier = nullptr;
 
-	if (CachedPlayerController == nullptr)
+	if (CachedPlayer == nullptr)
 	{
-		CachedPlayerController = InPC;
+		CachedPlayer = InPC->Player;
 		PlayerModifier = DaySequenceModifier;
 	}
 	else if (IsSplitscreenSupported())
 	{
         PlayerModifier = DuplicateObject<UDaySequenceModifierComponent>(DaySequenceModifier, this, TEXT("AdditionalPlayerModifier"));
-		AdditionalPlayers.FindOrAdd(InPC) = PlayerModifier;
+		AdditionalPlayers.FindOrAdd(InPC->Player) = PlayerModifier;
 	}
 	else
 	{
-		// If CachedPlayerController is not null (points to the first local player) and we don't support splitscreen, just track the player.
-		AdditionalPlayers.FindOrAdd(InPC) = nullptr;
+		// If CachedPlayer is not null and we don't support splitscreen, just track the player.
+		AdditionalPlayers.FindOrAdd(InPC->Player) = nullptr;
 	}
 
 	if (PlayerModifier)
@@ -276,7 +303,7 @@ void ADaySequenceModifierVolume::CreatePlayer(APlayerController* InPC)
 		// If we get here, we have multiple local players and do not have splitscreen support enabled and our CVar state allows us to disable the managed modifiers.
 		// In this case, we will just disable all modifiers.
 		
-		for (TTuple AdditionalPlayer : AdditionalPlayers)
+		for (TPair<TObjectPtr<UPlayer>, TObjectPtr<UDaySequenceModifierComponent>> AdditionalPlayer : AdditionalPlayers)
 		{
 			if (AdditionalPlayer.Value)
 			{
diff --git a/Engine/Plugins/Experimental/DaySequence/Source/DaySequence/Public/Actors/DaySequenceModifierVolume.h b/Engine/Plugins/Experimental/DaySequence/Source/DaySequence/Public/Actors/DaySequenceModifierVolume.h
--- a/Engine/Plugins/Experimental/DaySequence/Source/DaySequence/Public/Actors/DaySequenceModifierVolume.h
+++ b/Engine/Plugins/Experimental/DaySequence/Source/DaySequence/Public/Actors/DaySequenceModifierVolume.h
@@ -14,6 +14,7 @@ class APawn;
 
 class UBoxComponent;
 class UDaySequenceModifierComponent;
+class UPlayer;
 
 UCLASS(MinimalAPI, Blueprintable)
 class ADaySequenceModifierVolume : public AActor
@@ -62,8 +63,12 @@ protected:
 	UPROPERTY(Transient, DuplicateTransient)
     TObjectPtr<ADaySequenceActor> DaySequenceActor;
 	
-	UPROPERTY(Transient, DuplicateTransient)
+	UE_DEPRECATED(5.8, "Replaced with CachedPlayer")
+	UPROPERTY(Transient, DuplicateTransient, meta = (DeprecatedProperty, DeprecationMessage = "Replaced with CachedPlayer"))
 	TObjectPtr<APlayerController> CachedPlayerController;
+
+	UPROPERTY(Transient, DuplicateTransient)
+	TObjectPtr<UPlayer> CachedPlayer;
 	
 	FDelegateHandle ReplayScrubbedHandle;
 
@@ -73,9 +78,13 @@ private:
 	UPROPERTY(EditAnywhere, AdvancedDisplay, Category = "Day Sequence", DisplayName = "Enable Experimental Splitscreen Support")
 	bool bEnableSplitscreenSupport;
 
-	/** Transient modifier components associated with additional local players (primarily for split screen support) */
+	/**
+	 * Transient modifier components associated with additional local players (primarily for split screen support)
+	 * Key: Local player
+	 * Value: Optional transient modifier component
+	 */
 	UPROPERTY(Transient, DuplicateTransient)
-	TMap<TObjectPtr<APlayerController>, TObjectPtr<UDaySequenceModifierComponent>> AdditionalPlayers;
+	TMap<TObjectPtr<UPlayer>, TObjectPtr<UDaySequenceModifierComponent>> AdditionalPlayers;
 
 	FDelegateHandle ActorSpawnedHandle;
 };
