DaySequenceModifierComponent: Add blend time accessors

#Synced-CL 50710442

#Synced-CL 50710442

diff --git a/Engine/Plugins/Experimental/DaySequence/Source/DaySequence/Private/DaySequenceModifierComponent.cpp b/Engine/Plugins/Experimental/DaySequence/Source/DaySequence/Private/DaySequenceModifierComponent.cpp
--- a/Engine/Plugins/Experimental/DaySequence/Source/DaySequence/Private/DaySequenceModifierComponent.cpp
+++ b/Engine/Plugins/Experimental/DaySequence/Source/DaySequence/Private/DaySequenceModifierComponent.cpp
@@ -1073,6 +1073,16 @@ void UDaySequenceModifierComponent::SetBlendTarget(APlayerController* InActor)
 	WeakBlendTarget = InActor;
 }
 
+void UDaySequenceModifierComponent::SetBlendTime(float TimeInSeconds)
+{
+	BlendTime = TimeInSeconds;
+}
+
+float UDaySequenceModifierComponent::GetBlendTime() const
+{
+	return BlendTime;
+}
+
 void UDaySequenceModifierComponent::SetUserBlendWeight(float Weight)
 {
 	UserBlendWeight = FMath::Clamp(Weight, 0.f, 1.f);
diff --git a/Engine/Plugins/Experimental/DaySequence/Source/DaySequence/Public/DaySequenceModifierComponent.h b/Engine/Plugins/Experimental/DaySequence/Source/DaySequence/Public/DaySequenceModifierComponent.h
--- a/Engine/Plugins/Experimental/DaySequence/Source/DaySequence/Public/DaySequenceModifierComponent.h
+++ b/Engine/Plugins/Experimental/DaySequence/Source/DaySequence/Public/DaySequenceModifierComponent.h
@@ -200,6 +200,10 @@ public:
 	DAYSEQUENCE_API void SetBlendPolicy(EDaySequenceModifierUserBlendPolicy NewPolicy);
 	DAYSEQUENCE_API EDaySequenceModifierUserBlendPolicy GetBlendPolicy() const;
 
+	/** Sets the time to blend towards the target weight. */
+	DAYSEQUENCE_API void SetBlendTime(float TimeInSeconds);
+	DAYSEQUENCE_API float GetBlendTime() const;
+
 	/** Sets a custom blend weight for volume based blends. Final weight depends on BlendPolicy. */
 	DAYSEQUENCE_API void SetUserBlendWeight(float Weight);
 	DAYSEQUENCE_API float GetUserBlendWeight() const;
