DaySequence: Fixed Time mode blending popping for DaySequenceModifierVolumes.

diff --git a/Engine/Plugins/Experimental/DaySequence/Source/DaySequence/Private/DaySequenceModifierComponent.cpp b/Engine/Plugins/Experimental/DaySequence/Source/DaySequence/Private/DaySequenceModifierComponent.cpp
--- a/Engine/Plugins/Experimental/DaySequence/Source/DaySequence/Private/DaySequenceModifierComponent.cpp
+++ b/Engine/Plugins/Experimental/DaySequence/Source/DaySequence/Private/DaySequenceModifierComponent.cpp
@@ -1256,8 +1256,8 @@ float UDaySequenceModifierComponent::UpdateInternalBlendWeight()
 			const float DeltaTimeRatio = FMath::Clamp(DeltaTime / FMath::Max(BlendTime, UE_SMALL_NUMBER), 0.f, 1.f);
 			const float DeltaWeight = DeltaWeightRange * DeltaTimeRatio;
 
-			const float MinWeight = DeltaWeight > 0.f ? TimedBlendingStartWeight : TimedBlendingTargetWeight;
-			const float MaxWeight = DeltaWeight > 0.f ? TimedBlendingTargetWeight : TimedBlendingStartWeight;
+			const float MinWeight = DeltaWeightRange >= 0.f ? TimedBlendingStartWeight : TimedBlendingTargetWeight;
+			const float MaxWeight = DeltaWeightRange >= 0.f ? TimedBlendingTargetWeight : TimedBlendingStartWeight;
 			InternalBlendWeight = FMath::Clamp(TimedBlendingStartWeight + DeltaWeight, MinWeight, MaxWeight);
 			
 			UE_LOG(LogDaySequence, VeryVerbose, TEXT("[BlendWeight]: Start = %f, Target = %f, DeltaTime = %f, Weight = %f"), TimedBlendingStartWeight, TimedBlendingTargetWeight, DeltaTime, InternalBlendWeight);
