DaySequenceModifierVolume: Disable spatial loading by default

#Synced-CL 52359139

diff --git a/Engine/Plugins/Experimental/DaySequence/Source/DaySequence/Private/Actors/DaySequenceModifierVolume.cpp b/Engine/Plugins/Experimental/DaySequence/Source/DaySequence/Private/Actors/DaySequenceModifierVolume.cpp
--- a/Engine/Plugins/Experimental/DaySequence/Source/DaySequence/Private/Actors/DaySequenceModifierVolume.cpp
+++ b/Engine/Plugins/Experimental/DaySequence/Source/DaySequence/Private/Actors/DaySequenceModifierVolume.cpp
@@ -42,6 +42,12 @@ ADaySequenceModifierVolume::ADaySequenceModifierVolume(const FObjectInitializer&
 , bEnableSplitscreenSupport(false)
 {
 	PrimaryActorTick.bCanEverTick = false;
+
+#if WITH_EDITORONLY_DATA
+	// This is not generally safe. It is only safe if the modifier is in volume mode and
+	// the spatial loading threshold distance is greater than the radius of the volume.
+	bIsSpatiallyLoaded = false;
+#endif
 	
 	DaySequenceModifier = CreateDefaultSubobject<UDaySequenceModifierComponent>(TEXT("DaySequenceModifier"));
 	SetRootComponent(DaySequenceModifier);
